Another template.

diff --git a/Templates/TP_ThirdPerson/Config/DefaultInput.ini b/Templates/TP_ThirdPerson/Config/DefaultInput.ini
--- a/Templates/TP_ThirdPerson/Config/DefaultInput.ini
+++ b/Templates/TP_ThirdPerson/Config/DefaultInput.ini
@@ -19,3 +19,5 @@
 
 +AxisMappings=(AxisName="LookUpRate", Key=Gamepad_RightY, Scale=1.f)
 +AxisMappings=(AxisName="LookUp", Key=MouseY, Scale=-1.f)
+
++ActionMappings=(ActionName="ResetVR", Key=R)
diff --git a/Templates/TP_ThirdPerson/Source/TP_ThirdPerson/TP_ThirdPerson.Build.cs b/Templates/TP_ThirdPerson/Source/TP_ThirdPerson/TP_ThirdPerson.Build.cs
--- a/Templates/TP_ThirdPerson/Source/TP_ThirdPerson/TP_ThirdPerson.Build.cs
+++ b/Templates/TP_ThirdPerson/Source/TP_ThirdPerson/TP_ThirdPerson.Build.cs
@@ -6,6 +6,6 @@ public class TP_ThirdPerson : ModuleRules
 {
 	public TP_ThirdPerson(TargetInfo Target)
 	{
-		PublicDependencyModuleNames.AddRange(new string[] { "Core", "CoreUObject", "Engine", "InputCore" });
+		PublicDependencyModuleNames.AddRange(new string[] { "Core", "CoreUObject", "Engine", "InputCore", "HeadMountedDisplay" });
 	}
 }
diff --git a/Templates/TP_ThirdPerson/Source/TP_ThirdPerson/TP_ThirdPersonCharacter.cpp b/Templates/TP_ThirdPerson/Source/TP_ThirdPerson/TP_ThirdPersonCharacter.cpp
--- a/Templates/TP_ThirdPerson/Source/TP_ThirdPerson/TP_ThirdPersonCharacter.cpp
+++ b/Templates/TP_ThirdPerson/Source/TP_ThirdPerson/TP_ThirdPersonCharacter.cpp
@@ -1,6 +1,7 @@
 // Copyright 1998-2016 Epic Games, Inc. All Rights Reserved.
 
 #include "TP_ThirdPerson.h"
+#include "Kismet/HeadMountedDisplayFunctionLibrary.h"
 #include "TP_ThirdPersonCharacter.h"
 
 //////////////////////////////////////////////////////////////////////////
@@ -65,9 +66,17 @@ void ATP_ThirdPersonCharacter::SetupPlayerInputComponent(class UInputComponent*
 	// handle touch devices
 	PlayerInputComponent->BindTouch(IE_Pressed, this, &ATP_ThirdPersonCharacter::TouchStarted);
 	PlayerInputComponent->BindTouch(IE_Released, this, &ATP_ThirdPersonCharacter::TouchStopped);
+
+	// VR headset functionality
+	PlayerInputComponent->BindAction("ResetVR", IE_Pressed, this, &ATP_ThirdPersonCharacter::OnResetVR);
 }
 
 
+void ATP_ThirdPersonCharacter::OnResetVR()
+{
+	UHeadMountedDisplayFunctionLibrary::ResetOrientationAndPosition();
+}
+
 void ATP_ThirdPersonCharacter::TouchStarted(ETouchIndex::Type FingerIndex, FVector Location)
 {
 	// jump, but only on the first touch
diff --git a/Templates/TP_ThirdPerson/Source/TP_ThirdPerson/TP_ThirdPersonCharacter.h b/Templates/TP_ThirdPerson/Source/TP_ThirdPerson/TP_ThirdPersonCharacter.h
--- a/Templates/TP_ThirdPerson/Source/TP_ThirdPerson/TP_ThirdPersonCharacter.h
+++ b/Templates/TP_ThirdPerson/Source/TP_ThirdPerson/TP_ThirdPersonCharacter.h
@@ -28,6 +28,9 @@ public:
 
 protected:
 
+	/** Resets HMD orientation in VR. */
+	void OnResetVR();
+
 	/** Called for forwards/backward input */
 	void MoveForward(float Value);
 
@@ -54,7 +57,7 @@ protected:
 
 protected:
 	// APawn interface
-	virtual void SetupPlayerInputComponent(class UInputComponent* InputComponent) override;
+	virtual void SetupPlayerInputComponent(class UInputComponent* PlayerInputComponent) override;
 	// End of APawn interface
 
 public:
