Fixed BP and Code characters so that the user can touch to jump while moving with the virtual joystick

diff --git a/Templates/TP_ThirdPerson/Source/TP_ThirdPerson/TP_ThirdPersonCharacter.cpp b/Templates/TP_ThirdPerson/Source/TP_ThirdPerson/TP_ThirdPersonCharacter.cpp
--- a/Templates/TP_ThirdPerson/Source/TP_ThirdPerson/TP_ThirdPersonCharacter.cpp
+++ b/Templates/TP_ThirdPerson/Source/TP_ThirdPerson/TP_ThirdPersonCharacter.cpp
@@ -79,19 +79,12 @@ void ATP_ThirdPersonCharacter::OnResetVR()
 
 void ATP_ThirdPersonCharacter::TouchStarted(ETouchIndex::Type FingerIndex, FVector Location)
 {
-	// jump, but only on the first touch
-	if (FingerIndex == ETouchIndex::Touch1)
-	{
 		Jump();
-	}
 }
 
 void ATP_ThirdPersonCharacter::TouchStopped(ETouchIndex::Type FingerIndex, FVector Location)
 {
-	if (FingerIndex == ETouchIndex::Touch1)
-	{
 		StopJumping();
-	}
 }
 
 void ATP_ThirdPersonCharacter::TurnAtRate(float Rate)
