Removed VR functionality from TP_ThirdPerson and TP_ThirdPersonBP. Cleaned up Blueprints and fixed comment typos.

diff --git a/Templates/TP_ThirdPerson/Config/config.ini b/Templates/TP_ThirdPerson/Config/config.ini
--- a/Templates/TP_ThirdPerson/Config/config.ini
+++ b/Templates/TP_ThirdPerson/Config/config.ini
@@ -1,7 +1,6 @@
 [InputSettings]
 +ActionMappings=(ActionName="Jump", Key=SpaceBar)
 +ActionMappings=(ActionName="Jump", Key=Gamepad_FaceButton_Bottom)
-+ActionMappings=(ActionName="ResetVR", Key=R)
 
 +AxisMappings=(AxisName="MoveForward", Key=W, Scale=1.f)
 +AxisMappings=(AxisName="MoveForward", Key=S, Scale=-1.f)
diff --git a/Templates/TP_ThirdPerson/Source/TP_ThirdPerson/TP_ThirdPersonCharacter.cpp b/Templates/TP_ThirdPerson/Source/TP_ThirdPerson/TP_ThirdPersonCharacter.cpp
--- a/Templates/TP_ThirdPerson/Source/TP_ThirdPerson/TP_ThirdPersonCharacter.cpp
+++ b/Templates/TP_ThirdPerson/Source/TP_ThirdPerson/TP_ThirdPersonCharacter.cpp
@@ -1,7 +1,6 @@
 // Copyright Epic Games, Inc. All Rights Reserved.
 
 #include "TP_ThirdPersonCharacter.h"
-#include "HeadMountedDisplayFunctionLibrary.h"
 #include "Camera/CameraComponent.h"
 #include "Components/CapsuleComponent.h"
 #include "Components/InputComponent.h"
@@ -44,7 +43,7 @@ ATP_ThirdPersonCharacter::ATP_ThirdPersonCharacter()
 	FollowCamera->bUsePawnControlRotation = false; // Camera does not rotate relative to arm
 
 	// Note: The skeletal mesh and anim blueprint references on the Mesh component (inherited from Character) 
-	// are set in the derived blueprint asset named MyCharacter (to avoid direct content references in C++)
+	// are set in the derived blueprint asset named ThirdPersonCharacter (to avoid direct content references in C++)
 }
 
 //////////////////////////////////////////////////////////////////////////
@@ -71,21 +70,6 @@ void ATP_ThirdPersonCharacter::SetupPlayerInputComponent(class UInputComponent*
 	// handle touch devices
 	PlayerInputComponent->BindTouch(IE_Pressed, this, &ATP_ThirdPersonCharacter::TouchStarted);
 	PlayerInputComponent->BindTouch(IE_Released, this, &ATP_ThirdPersonCharacter::TouchStopped);
-
-	// VR headset functionality
-	PlayerInputComponent->BindAction("ResetVR", IE_Pressed, this, &ATP_ThirdPersonCharacter::OnResetVR);
-}
-
-
-void ATP_ThirdPersonCharacter::OnResetVR()
-{
-	// If TP_ThirdPerson is added to a project via 'Add Feature' in the Unreal Editor the dependency on HeadMountedDisplay in TP_ThirdPerson.Build.cs is not automatically propagated
-	// and a linker error will result.
-	// You will need to either:
-	//		Add "HeadMountedDisplay" to [YourProject].Build.cs PublicDependencyModuleNames in order to build successfully (appropriate if supporting VR).
-	// or:
-	//		Comment or delete the call to ResetOrientationAndPosition below (appropriate if not supporting VR)
-	UHeadMountedDisplayFunctionLibrary::ResetOrientationAndPosition();
 }
 
 void ATP_ThirdPersonCharacter::TouchStarted(ETouchIndex::Type FingerIndex, FVector Location)
diff --git a/Templates/TP_ThirdPerson/Source/TP_ThirdPerson/TP_ThirdPersonCharacter.h b/Templates/TP_ThirdPerson/Source/TP_ThirdPerson/TP_ThirdPersonCharacter.h
--- a/Templates/TP_ThirdPerson/Source/TP_ThirdPerson/TP_ThirdPersonCharacter.h
+++ b/Templates/TP_ThirdPerson/Source/TP_ThirdPerson/TP_ThirdPersonCharacter.h
@@ -31,9 +31,6 @@ public:
 
 protected:
 
-	/** Resets HMD orientation in VR. */
-	void OnResetVR();
-
 	/** Called for forwards/backward input */
 	void MoveForward(float Value);
 
diff --git a/Templates/TP_ThirdPersonBP/Config/config.ini b/Templates/TP_ThirdPersonBP/Config/config.ini
--- a/Templates/TP_ThirdPersonBP/Config/config.ini
+++ b/Templates/TP_ThirdPersonBP/Config/config.ini
@@ -1,7 +1,6 @@
 [InputSettings]
 +ActionMappings=(ActionName="Jump", Key=SpaceBar)
 +ActionMappings=(ActionName="Jump", Key=Gamepad_FaceButton_Bottom)
-+ActionMappings=(ActionName="ResetVR", Key=R)
 
 +AxisMappings=(AxisName="MoveForward", Key=W, Scale=1.f)
 +AxisMappings=(AxisName="MoveForward", Key=S, Scale=-1.f)
