Remove unnecessary CastChecks from the template projects. If people make a new project and revert to the old legacy system, that should be acceptable behavior and not cause a crash. We should just log instead :)

Clean up some of the forward declarations in the header files and Cpp files, there were a few types that were not used or had forward decls in the cpp

Define a log category for the First and Third person template so that we can have some messaging about stuff not working

diff --git a/Templates/TP_FirstPerson/Source/TP_FirstPerson/TP_FirstPersonCharacter.cpp b/Templates/TP_FirstPerson/Source/TP_FirstPerson/TP_FirstPersonCharacter.cpp
--- a/Templates/TP_FirstPerson/Source/TP_FirstPerson/TP_FirstPersonCharacter.cpp
+++ b/Templates/TP_FirstPerson/Source/TP_FirstPerson/TP_FirstPersonCharacter.cpp
@@ -5,9 +5,13 @@
 #include "Animation/AnimInstance.h"
 #include "Camera/CameraComponent.h"
 #include "Components/CapsuleComponent.h"
+#include "Components/SkeletalMeshComponent.h"
 #include "EnhancedInputComponent.h"
 #include "EnhancedInputSubsystems.h"
+#include "InputActionValue.h"
+#include "Engine/LocalPlayer.h"
 
+DEFINE_LOG_CATEGORY(LogTemplateCharacter);
 
 //////////////////////////////////////////////////////////////////////////
 // ATP_FirstPersonCharacter
@@ -42,7 +46,7 @@ void ATP_FirstPersonCharacter::BeginPlay()
 	// Call the base class  
 	Super::BeginPlay();
 
-	//Add Input Mapping Context
+	// Add Input Mapping Context
 	if (APlayerController* PlayerController = Cast<APlayerController>(Controller))
 	{
 		if (UEnhancedInputLocalPlayerSubsystem* Subsystem = ULocalPlayer::GetSubsystem<UEnhancedInputLocalPlayerSubsystem>(PlayerController->GetLocalPlayer()))
@@ -55,21 +59,25 @@ void ATP_FirstPersonCharacter::BeginPlay()
 
 //////////////////////////////////////////////////////////////////////////// Input
 
-void ATP_FirstPersonCharacter::SetupPlayerInputComponent(class UInputComponent* PlayerInputComponent)
+void ATP_FirstPersonCharacter::SetupPlayerInputComponent(UInputComponent* PlayerInputComponent)
 {
 	// Set up action bindings
-	if (UEnhancedInputComponent* EnhancedInputComponent = CastChecked<UEnhancedInputComponent>(PlayerInputComponent))
+	if (UEnhancedInputComponent* EnhancedInputComponent = Cast<UEnhancedInputComponent>(PlayerInputComponent))
 	{
-		//Jumping
+		// Jumping
 		EnhancedInputComponent->BindAction(JumpAction, ETriggerEvent::Triggered, this, &ACharacter::Jump);
 		EnhancedInputComponent->BindAction(JumpAction, ETriggerEvent::Completed, this, &ACharacter::StopJumping);
 
-		//Moving
+		// Moving
 		EnhancedInputComponent->BindAction(MoveAction, ETriggerEvent::Triggered, this, &ATP_FirstPersonCharacter::Move);
 
-		//Looking
+		// Looking
 		EnhancedInputComponent->BindAction(LookAction, ETriggerEvent::Triggered, this, &ATP_FirstPersonCharacter::Look);
 	}
+	else
+	{
+		UE_LOG(LogTemplateCharacter, Error, TEXT("'%s' Failed to find an Enhanced Input Component! This template is built to use the Enhanced Input system. If you intend to use the legacy system, then you will need to update this C++ file."), *GetNameSafe(this));
+	}
 }
 
 
diff --git a/Templates/TP_FirstPerson/Source/TP_FirstPerson/TP_FirstPersonCharacter.h b/Templates/TP_FirstPerson/Source/TP_FirstPerson/TP_FirstPersonCharacter.h
--- a/Templates/TP_FirstPerson/Source/TP_FirstPerson/TP_FirstPersonCharacter.h
+++ b/Templates/TP_FirstPerson/Source/TP_FirstPerson/TP_FirstPersonCharacter.h
@@ -4,15 +4,17 @@
 
 #include "CoreMinimal.h"
 #include "GameFramework/Character.h"
-#include "InputActionValue.h"
+#include "Logging/LogMacros.h"
 #include "TP_FirstPersonCharacter.generated.h"
 
 class UInputComponent;
 class USkeletalMeshComponent;
-class USceneComponent;
 class UCameraComponent;
-class UAnimMontage;
-class USoundBase;
+class UInputAction;
+class UInputMappingContext;
+struct FInputActionValue;
+
+DECLARE_LOG_CATEGORY_EXTERN(LogTemplateCharacter, Log, All);
 
 UCLASS(config=Game)
 class ATP_FirstPersonCharacter : public ACharacter
@@ -29,16 +31,15 @@ class ATP_FirstPersonCharacter : public ACharacter
 
 	/** MappingContext */
 	UPROPERTY(EditAnywhere, BlueprintReadOnly, Category=Input, meta=(AllowPrivateAccess = "true"))
-	class UInputMappingContext* DefaultMappingContext;
+	UInputMappingContext* DefaultMappingContext;
 
 	/** Jump Input Action */
 	UPROPERTY(EditAnywhere, BlueprintReadOnly, Category=Input, meta=(AllowPrivateAccess = "true"))
-	class UInputAction* JumpAction;
+	UInputAction* JumpAction;
 
 	/** Move Input Action */
 	UPROPERTY(EditAnywhere, BlueprintReadOnly, Category=Input, meta=(AllowPrivateAccess = "true"))
-	class UInputAction* MoveAction;
-
+	UInputAction* MoveAction;
 	
 public:
 	ATP_FirstPersonCharacter();
diff --git a/Templates/TP_ThirdPerson/Source/TP_ThirdPerson/TP_ThirdPersonCharacter.cpp b/Templates/TP_ThirdPerson/Source/TP_ThirdPerson/TP_ThirdPersonCharacter.cpp
--- a/Templates/TP_ThirdPerson/Source/TP_ThirdPerson/TP_ThirdPersonCharacter.cpp
+++ b/Templates/TP_ThirdPerson/Source/TP_ThirdPerson/TP_ThirdPersonCharacter.cpp
@@ -1,15 +1,17 @@
 // Copyright Epic Games, Inc. All Rights Reserved.
 
 #include "TP_ThirdPersonCharacter.h"
+#include "Engine/LocalPlayer.h"
 #include "Camera/CameraComponent.h"
 #include "Components/CapsuleComponent.h"
-#include "Components/InputComponent.h"
 #include "GameFramework/CharacterMovementComponent.h"
-#include "GameFramework/Controller.h"
 #include "GameFramework/SpringArmComponent.h"
+#include "GameFramework/Controller.h"
 #include "EnhancedInputComponent.h"
 #include "EnhancedInputSubsystems.h"
+#include "InputActionValue.h"
 
+DEFINE_LOG_CATEGORY(LogTemplateCharacter);
 
 //////////////////////////////////////////////////////////////////////////
 // ATP_ThirdPersonCharacter
@@ -69,23 +71,25 @@ void ATP_ThirdPersonCharacter::BeginPlay()
 //////////////////////////////////////////////////////////////////////////
 // Input
 
-void ATP_ThirdPersonCharacter::SetupPlayerInputComponent(class UInputComponent* PlayerInputComponent)
+void ATP_ThirdPersonCharacter::SetupPlayerInputComponent(UInputComponent* PlayerInputComponent)
 {
 	// Set up action bindings
-	if (UEnhancedInputComponent* EnhancedInputComponent = CastChecked<UEnhancedInputComponent>(PlayerInputComponent)) {
+	if (UEnhancedInputComponent* EnhancedInputComponent = Cast<UEnhancedInputComponent>(PlayerInputComponent)) {
 		
-		//Jumping
+		// Jumping
 		EnhancedInputComponent->BindAction(JumpAction, ETriggerEvent::Triggered, this, &ACharacter::Jump);
 		EnhancedInputComponent->BindAction(JumpAction, ETriggerEvent::Completed, this, &ACharacter::StopJumping);
 
-		//Moving
+		// Moving
 		EnhancedInputComponent->BindAction(MoveAction, ETriggerEvent::Triggered, this, &ATP_ThirdPersonCharacter::Move);
 
-		//Looking
+		// Looking
 		EnhancedInputComponent->BindAction(LookAction, ETriggerEvent::Triggered, this, &ATP_ThirdPersonCharacter::Look);
-
 	}
-
+	else
+	{
+		UE_LOG(LogTemplateCharacter, Error, TEXT("'%s' Failed to find an Enhanced Input component! This template is built to use the Enhanced Input system. If you intend to use the legacy system, then you will need to update this C++ file."), *GetNameSafe(this));
+	}
 }
 
 void ATP_ThirdPersonCharacter::Move(const FInputActionValue& Value)
@@ -122,8 +126,4 @@ void ATP_ThirdPersonCharacter::Look(const FInputActionValue& Value)
 		AddControllerYawInput(LookAxisVector.X);
 		AddControllerPitchInput(LookAxisVector.Y);
 	}
-}
-
-
-
-
+}
\ No newline at end of file
diff --git a/Templates/TP_ThirdPerson/Source/TP_ThirdPerson/TP_ThirdPersonCharacter.h b/Templates/TP_ThirdPerson/Source/TP_ThirdPerson/TP_ThirdPersonCharacter.h
--- a/Templates/TP_ThirdPerson/Source/TP_ThirdPerson/TP_ThirdPersonCharacter.h
+++ b/Templates/TP_ThirdPerson/Source/TP_ThirdPerson/TP_ThirdPersonCharacter.h
@@ -4,9 +4,16 @@
 
 #include "CoreMinimal.h"
 #include "GameFramework/Character.h"
-#include "InputActionValue.h"
+#include "Logging/LogMacros.h"
 #include "TP_ThirdPersonCharacter.generated.h"
 
+class USpringArmComponent;
+class UCameraComponent;
+class UInputMappingContext;
+class UInputAction;
+struct FInputActionValue;
+
+DECLARE_LOG_CATEGORY_EXTERN(LogTemplateCharacter, Log, All);
 
 UCLASS(config=Game)
 class ATP_ThirdPersonCharacter : public ACharacter
@@ -15,27 +22,27 @@ class ATP_ThirdPersonCharacter : public ACharacter
 
 	/** Camera boom positioning the camera behind the character */
 	UPROPERTY(VisibleAnywhere, BlueprintReadOnly, Category = Camera, meta = (AllowPrivateAccess = "true"))
-	class USpringArmComponent* CameraBoom;
+	USpringArmComponent* CameraBoom;
 
 	/** Follow camera */
 	UPROPERTY(VisibleAnywhere, BlueprintReadOnly, Category = Camera, meta = (AllowPrivateAccess = "true"))
-	class UCameraComponent* FollowCamera;
+	UCameraComponent* FollowCamera;
 	
 	/** MappingContext */
 	UPROPERTY(EditAnywhere, BlueprintReadOnly, Category = Input, meta = (AllowPrivateAccess = "true"))
-	class UInputMappingContext* DefaultMappingContext;
+	UInputMappingContext* DefaultMappingContext;
 
 	/** Jump Input Action */
 	UPROPERTY(EditAnywhere, BlueprintReadOnly, Category = Input, meta = (AllowPrivateAccess = "true"))
-	class UInputAction* JumpAction;
+	UInputAction* JumpAction;
 
 	/** Move Input Action */
 	UPROPERTY(EditAnywhere, BlueprintReadOnly, Category = Input, meta = (AllowPrivateAccess = "true"))
-	class UInputAction* MoveAction;
+	UInputAction* MoveAction;
 
 	/** Look Input Action */
 	UPROPERTY(EditAnywhere, BlueprintReadOnly, Category = Input, meta = (AllowPrivateAccess = "true"))
-	class UInputAction* LookAction;
+	UInputAction* LookAction;
 
 public:
 	ATP_ThirdPersonCharacter();
diff --git a/Templates/TP_TopDown/Source/TP_TopDown/TP_TopDownPlayerController.cpp b/Templates/TP_TopDown/Source/TP_TopDown/TP_TopDownPlayerController.cpp
--- a/Templates/TP_TopDown/Source/TP_TopDown/TP_TopDownPlayerController.cpp
+++ b/Templates/TP_TopDown/Source/TP_TopDown/TP_TopDownPlayerController.cpp
@@ -8,7 +8,11 @@
 #include "TP_TopDownCharacter.h"
 #include "Engine/World.h"
 #include "EnhancedInputComponent.h"
+#include "InputActionValue.h"
 #include "EnhancedInputSubsystems.h"
+#include "Engine/LocalPlayer.h"
+
+DEFINE_LOG_CATEGORY(LogTemplateCharacter);
 
 ATP_TopDownPlayerController::ATP_TopDownPlayerController()
 {
@@ -36,7 +40,7 @@ void ATP_TopDownPlayerController::SetupInputComponent()
 	Super::SetupInputComponent();
 
 	// Set up action bindings
-	if (UEnhancedInputComponent* EnhancedInputComponent = CastChecked<UEnhancedInputComponent>(InputComponent))
+	if (UEnhancedInputComponent* EnhancedInputComponent = Cast<UEnhancedInputComponent>(InputComponent))
 	{
 		// Setup mouse input events
 		EnhancedInputComponent->BindAction(SetDestinationClickAction, ETriggerEvent::Started, this, &ATP_TopDownPlayerController::OnInputStarted);
@@ -50,6 +54,10 @@ void ATP_TopDownPlayerController::SetupInputComponent()
 		EnhancedInputComponent->BindAction(SetDestinationTouchAction, ETriggerEvent::Completed, this, &ATP_TopDownPlayerController::OnTouchReleased);
 		EnhancedInputComponent->BindAction(SetDestinationTouchAction, ETriggerEvent::Canceled, this, &ATP_TopDownPlayerController::OnTouchReleased);
 	}
+	else
+	{
+		UE_LOG(LogTemplateCharacter, Error, TEXT("'%s' Failed to find an Enhanced Input Component! This template is built to use the Enhanced Input system. If you intend to use the legacy system, then you will need to update this C++ file."), *GetNameSafe(this));
+	}
 }
 
 void ATP_TopDownPlayerController::OnInputStarted()
diff --git a/Templates/TP_TopDown/Source/TP_TopDown/TP_TopDownPlayerController.h b/Templates/TP_TopDown/Source/TP_TopDown/TP_TopDownPlayerController.h
--- a/Templates/TP_TopDown/Source/TP_TopDown/TP_TopDownPlayerController.h
+++ b/Templates/TP_TopDown/Source/TP_TopDown/TP_TopDownPlayerController.h
@@ -5,11 +5,14 @@
 #include "CoreMinimal.h"
 #include "Templates/SubclassOf.h"
 #include "GameFramework/PlayerController.h"
-#include "InputActionValue.h"
 #include "TP_TopDownPlayerController.generated.h"
 
 /** Forward declaration to improve compiling times */
 class UNiagaraSystem;
+class UInputMappingContext;
+class UInputAction;
+
+DECLARE_LOG_CATEGORY_EXTERN(LogTemplateCharacter, Log, All);
 
 UCLASS()
 class ATP_TopDownPlayerController : public APlayerController
@@ -29,15 +32,15 @@ public:
 
 	/** MappingContext */
 	UPROPERTY(EditAnywhere, BlueprintReadOnly, Category=Input, meta=(AllowPrivateAccess = "true"))
-	class UInputMappingContext* DefaultMappingContext;
+	UInputMappingContext* DefaultMappingContext;
 	
 	/** Jump Input Action */
 	UPROPERTY(EditAnywhere, BlueprintReadOnly, Category=Input, meta=(AllowPrivateAccess = "true"))
-	class UInputAction* SetDestinationClickAction;
+	UInputAction* SetDestinationClickAction;
 
 	/** Jump Input Action */
 	UPROPERTY(EditAnywhere, BlueprintReadOnly, Category=Input, meta=(AllowPrivateAccess = "true"))
-	class UInputAction* SetDestinationTouchAction;
+	UInputAction* SetDestinationTouchAction;
 
 protected:
 	/** True if the controlled character should navigate to the mouse cursor. */
