Character in Third Person templates no longer constantly jumps if the button is held.
StopJumping() is also correctly called on button release.
Also fixes break jump deceleration issue from UE-173720

diff --git a/Templates/TP_ThirdPerson/Source/TP_ThirdPerson/TP_ThirdPersonCharacter.cpp b/Templates/TP_ThirdPerson/Source/TP_ThirdPerson/TP_ThirdPersonCharacter.cpp
--- a/Templates/TP_ThirdPerson/Source/TP_ThirdPerson/TP_ThirdPersonCharacter.cpp
+++ b/Templates/TP_ThirdPerson/Source/TP_ThirdPerson/TP_ThirdPersonCharacter.cpp
@@ -37,6 +37,7 @@ ATP_ThirdPersonCharacter::ATP_ThirdPersonCharacter()
 	GetCharacterMovement()->MaxWalkSpeed = 500.f;
 	GetCharacterMovement()->MinAnalogWalkSpeed = 20.f;
 	GetCharacterMovement()->BrakingDecelerationWalking = 2000.f;
+	GetCharacterMovement()->BrakingDecelerationFalling = 1500.0f;
 
 	// Create a camera boom (pulls in towards the player if there is a collision)
 	CameraBoom = CreateDefaultSubobject<USpringArmComponent>(TEXT("CameraBoom"));
@@ -77,7 +78,7 @@ void ATP_ThirdPersonCharacter::SetupPlayerInputComponent(UInputComponent* Player
 	if (UEnhancedInputComponent* EnhancedInputComponent = Cast<UEnhancedInputComponent>(PlayerInputComponent)) {
 		
 		// Jumping
-		EnhancedInputComponent->BindAction(JumpAction, ETriggerEvent::Triggered, this, &ACharacter::Jump);
+		EnhancedInputComponent->BindAction(JumpAction, ETriggerEvent::Started, this, &ACharacter::Jump);
 		EnhancedInputComponent->BindAction(JumpAction, ETriggerEvent::Completed, this, &ACharacter::StopJumping);
 
 		// Moving
