Fix an issue where the Third Person template would not have it's Input Mapping Context added in a splitscreen game because the owning Player Controller would not yet be available on Begin Play if you create a player with the Game Instance manually.

diff --git a/Templates/TP_FirstPerson/Source/TP_FirstPerson/TP_FirstPersonCharacter.cpp b/Templates/TP_FirstPerson/Source/TP_FirstPerson/TP_FirstPersonCharacter.cpp
--- a/Templates/TP_FirstPerson/Source/TP_FirstPerson/TP_FirstPersonCharacter.cpp
+++ b/Templates/TP_FirstPerson/Source/TP_FirstPerson/TP_FirstPersonCharacter.cpp
@@ -45,7 +45,12 @@ void ATP_FirstPersonCharacter::BeginPlay()
 {
 	// Call the base class  
 	Super::BeginPlay();
+}
+
+//////////////////////////////////////////////////////////////////////////// Input
 
+void ATP_FirstPersonCharacter::SetupPlayerInputComponent(UInputComponent* PlayerInputComponent)
+{
 	// Add Input Mapping Context
 	if (APlayerController* PlayerController = Cast<APlayerController>(Controller))
 	{
@@ -54,13 +59,11 @@ void ATP_FirstPersonCharacter::BeginPlay()
 			Subsystem->AddMappingContext(DefaultMappingContext, 0);
 		}
 	}
-
-}
-
-//////////////////////////////////////////////////////////////////////////// Input
-
-void ATP_FirstPersonCharacter::SetupPlayerInputComponent(UInputComponent* PlayerInputComponent)
-{
+	else
+	{
+		UE_LOG(LogTemplateCharacter, Error, TEXT("'%s' Failed to add an input mapping context because there is no owning Player Controller!"), *GetNameSafe(this));
+	}
+	
 	// Set up action bindings
 	if (UEnhancedInputComponent* EnhancedInputComponent = Cast<UEnhancedInputComponent>(PlayerInputComponent))
 	{
diff --git a/Templates/TP_ThirdPerson/Source/TP_ThirdPerson/TP_ThirdPersonCharacter.cpp b/Templates/TP_ThirdPerson/Source/TP_ThirdPerson/TP_ThirdPersonCharacter.cpp
--- a/Templates/TP_ThirdPerson/Source/TP_ThirdPerson/TP_ThirdPersonCharacter.cpp
+++ b/Templates/TP_ThirdPerson/Source/TP_ThirdPerson/TP_ThirdPersonCharacter.cpp
@@ -58,15 +58,6 @@ void ATP_ThirdPersonCharacter::BeginPlay()
 {
 	// Call the base class  
 	Super::BeginPlay();
-
-	//Add Input Mapping Context
-	if (APlayerController* PlayerController = Cast<APlayerController>(Controller))
-	{
-		if (UEnhancedInputLocalPlayerSubsystem* Subsystem = ULocalPlayer::GetSubsystem<UEnhancedInputLocalPlayerSubsystem>(PlayerController->GetLocalPlayer()))
-		{
-			Subsystem->AddMappingContext(DefaultMappingContext, 0);
-		}
-	}
 }
 
 //////////////////////////////////////////////////////////////////////////
@@ -74,6 +65,15 @@ void ATP_ThirdPersonCharacter::BeginPlay()
 
 void ATP_ThirdPersonCharacter::SetupPlayerInputComponent(UInputComponent* PlayerInputComponent)
 {
+	// Add Input Mapping Context
+	if (APlayerController* PlayerController = Cast<APlayerController>(GetController()))
+	{
+		if (UEnhancedInputLocalPlayerSubsystem* Subsystem = ULocalPlayer::GetSubsystem<UEnhancedInputLocalPlayerSubsystem>(PlayerController->GetLocalPlayer()))
+		{
+			Subsystem->AddMappingContext(DefaultMappingContext, 0);
+		}
+	}
+	
 	// Set up action bindings
 	if (UEnhancedInputComponent* EnhancedInputComponent = Cast<UEnhancedInputComponent>(PlayerInputComponent)) {
 		
diff --git a/Templates/TP_VehicleAdv/Source/TP_VehicleAdv/TP_VehicleAdvPlayerController.cpp b/Templates/TP_VehicleAdv/Source/TP_VehicleAdv/TP_VehicleAdvPlayerController.cpp
--- a/Templates/TP_VehicleAdv/Source/TP_VehicleAdv/TP_VehicleAdvPlayerController.cpp
+++ b/Templates/TP_VehicleAdv/Source/TP_VehicleAdv/TP_VehicleAdvPlayerController.cpp
@@ -10,20 +10,25 @@
 void ATP_VehicleAdvPlayerController::BeginPlay()
 {
 	Super::BeginPlay();
+	
+	// spawn the UI widget and add it to the viewport
+	VehicleUI = CreateWidget<UTP_VehicleAdvUI>(this, VehicleUIClass);
+
+	check(VehicleUI);
 
+	VehicleUI->AddToViewport();
+}
+
+void ATP_VehicleAdvPlayerController::SetupInputComponent()
+{
+	Super::SetupInputComponent();
+	
 	// get the enhanced input subsystem
 	if (UEnhancedInputLocalPlayerSubsystem* Subsystem = ULocalPlayer::GetSubsystem<UEnhancedInputLocalPlayerSubsystem>(GetLocalPlayer()))
 	{
 		// add the mapping context so we get controls
 		Subsystem->AddMappingContext(InputMappingContext, 0);
 	}
-
-	// spawn the UI widget and add it to the viewport
-	VehicleUI = CreateWidget<UTP_VehicleAdvUI>(this, VehicleUIClass);
-
-	check(VehicleUI);
-
-	VehicleUI->AddToViewport();
 }
 
 void ATP_VehicleAdvPlayerController::Tick(float Delta)
diff --git a/Templates/TP_VehicleAdv/Source/TP_VehicleAdv/TP_VehicleAdvPlayerController.h b/Templates/TP_VehicleAdv/Source/TP_VehicleAdv/TP_VehicleAdvPlayerController.h
--- a/Templates/TP_VehicleAdv/Source/TP_VehicleAdv/TP_VehicleAdvPlayerController.h
+++ b/Templates/TP_VehicleAdv/Source/TP_VehicleAdv/TP_VehicleAdvPlayerController.h
@@ -39,6 +39,7 @@ protected:
 protected:
 
 	virtual void BeginPlay() override;
+	virtual void SetupInputComponent() override;
 
 public:
 
