Project Templates - Moved input mapping context registration to NotifyControllerChanged() to avoid an issue with Simulate into Attach PC workflows.
Removed PlayerController classes made redundant by the IMC registration change.
Removed unneeded component ownership transfer from TP_WeaponComponent

#rnx

diff --git a/Templates/TP_FirstPerson/Source/TP_FirstPerson/TP_FirstPersonCharacter.cpp b/Templates/TP_FirstPerson/Source/TP_FirstPerson/TP_FirstPersonCharacter.cpp
--- a/Templates/TP_FirstPerson/Source/TP_FirstPerson/TP_FirstPersonCharacter.cpp
+++ b/Templates/TP_FirstPerson/Source/TP_FirstPerson/TP_FirstPersonCharacter.cpp
@@ -33,18 +33,25 @@ ATP_FirstPersonCharacter::ATP_FirstPersonCharacter()
 	Mesh1P->SetupAttachment(FirstPersonCameraComponent);
 	Mesh1P->bCastDynamicShadow = false;
 	Mesh1P->CastShadow = false;
-	//Mesh1P->SetRelativeRotation(FRotator(0.9f, -19.19f, 5.2f));
 	Mesh1P->SetRelativeLocation(FVector(-30.f, 0.f, -150.f));
 
 }
 
-void ATP_FirstPersonCharacter::BeginPlay()
+//////////////////////////////////////////////////////////////////////////// Input
+
+void ATP_FirstPersonCharacter::NotifyControllerChanged()
 {
-	// Call the base class  
-	Super::BeginPlay();
-}
+	Super::NotifyControllerChanged();
 
-//////////////////////////////////////////////////////////////////////////// Input
+	// Add Input Mapping Context
+	if (APlayerController* PlayerController = Cast<APlayerController>(Controller))
+	{
+		if (UEnhancedInputLocalPlayerSubsystem* Subsystem = ULocalPlayer::GetSubsystem<UEnhancedInputLocalPlayerSubsystem>(PlayerController->GetLocalPlayer()))
+		{
+			Subsystem->AddMappingContext(DefaultMappingContext, 0);
+		}
+	}
+}
 
 void ATP_FirstPersonCharacter::SetupPlayerInputComponent(UInputComponent* PlayerInputComponent)
 {	
diff --git a/Templates/TP_FirstPerson/Source/TP_FirstPerson/TP_FirstPersonCharacter.h b/Templates/TP_FirstPerson/Source/TP_FirstPerson/TP_FirstPersonCharacter.h
--- a/Templates/TP_FirstPerson/Source/TP_FirstPerson/TP_FirstPersonCharacter.h
+++ b/Templates/TP_FirstPerson/Source/TP_FirstPerson/TP_FirstPersonCharacter.h
@@ -29,6 +29,10 @@ class ATP_FirstPersonCharacter : public ACharacter
 	UPROPERTY(VisibleAnywhere, BlueprintReadOnly, Category = Camera, meta = (AllowPrivateAccess = "true"))
 	UCameraComponent* FirstPersonCameraComponent;
 
+	/** MappingContext */
+	UPROPERTY(EditAnywhere, BlueprintReadOnly, Category = Input, meta = (AllowPrivateAccess = "true"))
+	UInputMappingContext* DefaultMappingContext;
+
 	/** Jump Input Action */
 	UPROPERTY(EditAnywhere, BlueprintReadOnly, Category=Input, meta=(AllowPrivateAccess = "true"))
 	UInputAction* JumpAction;
@@ -36,18 +40,13 @@ class ATP_FirstPersonCharacter : public ACharacter
 	/** Move Input Action */
 	UPROPERTY(EditAnywhere, BlueprintReadOnly, Category=Input, meta=(AllowPrivateAccess = "true"))
 	UInputAction* MoveAction;
-	
-public:
-	ATP_FirstPersonCharacter();
 
-protected:
-	virtual void BeginPlay();
-
-public:
-		
 	/** Look Input Action */
 	UPROPERTY(EditAnywhere, BlueprintReadOnly, Category = Input, meta = (AllowPrivateAccess = "true"))
 	class UInputAction* LookAction;
+	
+public:
+	ATP_FirstPersonCharacter();
 
 protected:
 	/** Called for movement input */
@@ -58,6 +57,7 @@ protected:
 
 protected:
 	// APawn interface
+	virtual void NotifyControllerChanged() override;
 	virtual void SetupPlayerInputComponent(UInputComponent* InputComponent) override;
 	// End of APawn interface
 
diff --git a/Templates/TP_FirstPerson/Source/TP_FirstPerson/TP_FirstPersonPlayerController.cpp b/Templates/TP_FirstPerson/Source/TP_FirstPerson/TP_FirstPersonPlayerController.cpp
deleted file mode 100644
diff --git a/Templates/TP_FirstPerson/Source/TP_FirstPerson/TP_FirstPersonPlayerController.h b/Templates/TP_FirstPerson/Source/TP_FirstPerson/TP_FirstPersonPlayerController.h
deleted file mode 100644
diff --git a/Templates/TP_FirstPerson/Source/TP_FirstPerson/TP_WeaponComponent.cpp b/Templates/TP_FirstPerson/Source/TP_FirstPerson/TP_WeaponComponent.cpp
--- a/Templates/TP_FirstPerson/Source/TP_FirstPerson/TP_WeaponComponent.cpp
+++ b/Templates/TP_FirstPerson/Source/TP_FirstPerson/TP_WeaponComponent.cpp
@@ -80,10 +80,6 @@ bool UTP_WeaponComponent::AttachWeapon(ATP_FirstPersonCharacter* TargetCharacter
 	FAttachmentTransformRules AttachmentRules(EAttachmentRule::SnapToTarget, true);
 	AttachToComponent(Character->GetMesh1P(), AttachmentRules, FName(TEXT("GripPoint")));
 
-	// the weapon component is owned by the Pickup, so we need to 
-	// manually add it to the Character as an Instance Component to hand off ownership
-	Character->AddInstanceComponent(this);
-
 	// Set up action bindings
 	if (APlayerController* PlayerController = Cast<APlayerController>(Character->GetController()))
 	{
diff --git a/Templates/TP_ThirdPerson/Source/TP_ThirdPerson/TP_ThirdPersonCharacter.cpp b/Templates/TP_ThirdPerson/Source/TP_ThirdPerson/TP_ThirdPersonCharacter.cpp
--- a/Templates/TP_ThirdPerson/Source/TP_ThirdPerson/TP_ThirdPersonCharacter.cpp
+++ b/Templates/TP_ThirdPerson/Source/TP_ThirdPerson/TP_ThirdPersonCharacter.cpp
@@ -54,26 +54,25 @@ ATP_ThirdPersonCharacter::ATP_ThirdPersonCharacter()
 	// are set in the derived blueprint asset named ThirdPersonCharacter (to avoid direct content references in C++)
 }
 
-void ATP_ThirdPersonCharacter::BeginPlay()
-{
-	// Call the base class  
-	Super::BeginPlay();
-}
-
 //////////////////////////////////////////////////////////////////////////
 // Input
 
-void ATP_ThirdPersonCharacter::SetupPlayerInputComponent(UInputComponent* PlayerInputComponent)
+void ATP_ThirdPersonCharacter::NotifyControllerChanged()
 {
+	Super::NotifyControllerChanged();
+
 	// Add Input Mapping Context
-	if (APlayerController* PlayerController = Cast<APlayerController>(GetController()))
+	if (APlayerController* PlayerController = Cast<APlayerController>(Controller))
 	{
 		if (UEnhancedInputLocalPlayerSubsystem* Subsystem = ULocalPlayer::GetSubsystem<UEnhancedInputLocalPlayerSubsystem>(PlayerController->GetLocalPlayer()))
 		{
 			Subsystem->AddMappingContext(DefaultMappingContext, 0);
 		}
 	}
-	
+}
+
+void ATP_ThirdPersonCharacter::SetupPlayerInputComponent(UInputComponent* PlayerInputComponent)
+{
 	// Set up action bindings
 	if (UEnhancedInputComponent* EnhancedInputComponent = Cast<UEnhancedInputComponent>(PlayerInputComponent)) {
 		
@@ -127,4 +126,4 @@ void ATP_ThirdPersonCharacter::Look(const FInputActionValue& Value)
 		AddControllerYawInput(LookAxisVector.X);
 		AddControllerPitchInput(LookAxisVector.Y);
 	}
-}
\ No newline at end of file
+}
diff --git a/Templates/TP_ThirdPerson/Source/TP_ThirdPerson/TP_ThirdPersonCharacter.h b/Templates/TP_ThirdPerson/Source/TP_ThirdPerson/TP_ThirdPersonCharacter.h
--- a/Templates/TP_ThirdPerson/Source/TP_ThirdPerson/TP_ThirdPersonCharacter.h
+++ b/Templates/TP_ThirdPerson/Source/TP_ThirdPerson/TP_ThirdPersonCharacter.h
@@ -58,11 +58,10 @@ protected:
 			
 
 protected:
-	// APawn interface
+
+	virtual void NotifyControllerChanged() override;
+
 	virtual void SetupPlayerInputComponent(class UInputComponent* PlayerInputComponent) override;
-	
-	// To add mapping context
-	virtual void BeginPlay();
 
 public:
 	/** Returns CameraBoom subobject **/
