Templates - Fix for duplicate LOCTEXT on StateTreeUtility files.
Temporarily replaced LOCTEXT for FText::FromString
#rnx

diff --git a/Templates/TP_FirstPerson/Source/TP_FirstPerson/Variant_Shooter/AI/ShooterStateTreeUtility.cpp b/Templates/TP_FirstPerson/Source/TP_FirstPerson/Variant_Shooter/AI/ShooterStateTreeUtility.cpp
--- a/Templates/TP_FirstPerson/Source/TP_FirstPerson/Variant_Shooter/AI/ShooterStateTreeUtility.cpp
+++ b/Templates/TP_FirstPerson/Source/TP_FirstPerson/Variant_Shooter/AI/ShooterStateTreeUtility.cpp
@@ -10,8 +10,6 @@
 #include "ShooterAIController.h"
 #include "StateTreeAsyncExecutionContext.h"
 
-#define LOCTEXT_NAMESPACE "FirstPersonTemplate"
-
 bool FStateTreeLineOfSightToTargetCondition::TestCondition(FStateTreeExecutionContext& Context) const
 {
 	const FInstanceDataType& InstanceData = Context.GetInstanceData(*this);
@@ -74,7 +72,7 @@ bool FStateTreeLineOfSightToTargetCondition::TestCondition(FStateTreeExecutionCo
 #if WITH_EDITOR
 FText FStateTreeLineOfSightToTargetCondition::GetDescription(const FGuid& ID, FStateTreeDataView InstanceDataView, const IStateTreeBindingLookup& BindingLookup, EStateTreeNodeFormatting Formatting /*= EStateTreeNodeFormatting::Text*/) const
 {
-	return LOCTEXT("StateTreeConditionLineOfSightToTarget", "<b>Has Line of Sight</b>");
+	return FText::FromString("<b>Has Line of Sight</b>");
 }
 #endif
 
@@ -111,7 +109,7 @@ void FStateTreeFaceActorTask::ExitState(FStateTreeExecutionContext& Context, con
 #if WITH_EDITOR
 FText FStateTreeFaceActorTask::GetDescription(const FGuid& ID, FStateTreeDataView InstanceDataView, const IStateTreeBindingLookup& BindingLookup, EStateTreeNodeFormatting Formatting /*= EStateTreeNodeFormatting::Text*/) const
 {
-	return LOCTEXT("StateTreeTaskFaceActorDescription", "<b>Face Towards Actor</b>");
+	return FText::FromString("<b>Face Towards Actor</b>");
 }
 #endif // WITH_EDITOR
 
@@ -148,7 +146,7 @@ void FStateTreeFaceLocationTask::ExitState(FStateTreeExecutionContext& Context,
 #if WITH_EDITOR
 FText FStateTreeFaceLocationTask::GetDescription(const FGuid& ID, FStateTreeDataView InstanceDataView, const IStateTreeBindingLookup& BindingLookup, EStateTreeNodeFormatting Formatting /*= EStateTreeNodeFormatting::Text*/) const
 {
-	return LOCTEXT("StateTreeTaskFaceLocationDescription", "<b>Face Towards Location</b>");
+	return FText::FromString("<b>Face Towards Location</b>");
 }
 #endif // WITH_EDITOR
 
@@ -172,7 +170,7 @@ EStateTreeRunStatus FStateTreeSetRandomFloatTask::EnterState(FStateTreeExecution
 #if WITH_EDITOR
 FText FStateTreeSetRandomFloatTask::GetDescription(const FGuid& ID, FStateTreeDataView InstanceDataView, const IStateTreeBindingLookup& BindingLookup, EStateTreeNodeFormatting Formatting /*= EStateTreeNodeFormatting::Text*/) const
 {
-	return LOCTEXT("StateTreeTaskSetRandomFloatDescription", "<b>Set Random Float</b>");
+	return FText::FromString("<b>Set Random Float</b>");
 }
 #endif // WITH_EDITOR
 
@@ -209,7 +207,7 @@ void FStateTreeShootAtTargetTask::ExitState(FStateTreeExecutionContext& Context,
 #if WITH_EDITOR
 FText FStateTreeShootAtTargetTask::GetDescription(const FGuid& ID, FStateTreeDataView InstanceDataView, const IStateTreeBindingLookup& BindingLookup, EStateTreeNodeFormatting Formatting /*= EStateTreeNodeFormatting::Text*/) const
 {
-	return LOCTEXT("StateTreeTaskShootAtTargetDescription", "<b>Shoot at Target</b>");
+	return FText::FromString("<b>Shoot at Target</b>");
 }
 #endif // WITH_EDITOR
 
@@ -365,6 +363,6 @@ void FStateTreeSenseEnemiesTask::ExitState(FStateTreeExecutionContext& Context,
 #if WITH_EDITOR
 FText FStateTreeSenseEnemiesTask::GetDescription(const FGuid& ID, FStateTreeDataView InstanceDataView, const IStateTreeBindingLookup& BindingLookup, EStateTreeNodeFormatting Formatting /*= EStateTreeNodeFormatting::Text*/) const
 {
-	return LOCTEXT("StateTreeTaskSenseEnemiesDescription", "<b>Sense Enemies</b>");
+	return FText::FromString("<b>Sense Enemies</b>");
 }
 #endif // WITH_EDITOR
\ No newline at end of file
diff --git a/Templates/TP_ThirdPerson/Source/TP_ThirdPerson/Variant_Combat/AI/CombatStateTreeUtility.cpp b/Templates/TP_ThirdPerson/Source/TP_ThirdPerson/Variant_Combat/AI/CombatStateTreeUtility.cpp
--- a/Templates/TP_ThirdPerson/Source/TP_ThirdPerson/Variant_Combat/AI/CombatStateTreeUtility.cpp
+++ b/Templates/TP_ThirdPerson/Source/TP_ThirdPerson/Variant_Combat/AI/CombatStateTreeUtility.cpp
@@ -11,8 +11,6 @@
 #include "Kismet/GameplayStatics.h"
 #include "StateTreeAsyncExecutionContext.h"
 
-#define LOCTEXT_NAMESPACE "ThirdPersonTemplate"
-
 bool FStateTreeCharacterGroundedCondition::TestCondition(FStateTreeExecutionContext& Context) const
 {
 	const FInstanceDataType& InstanceData = Context.GetInstanceData(*this);
@@ -26,7 +24,7 @@ bool FStateTreeCharacterGroundedCondition::TestCondition(FStateTreeExecutionCont
 #if WITH_EDITOR
 FText FStateTreeCharacterGroundedCondition::GetDescription(const FGuid& ID, FStateTreeDataView InstanceDataView, const IStateTreeBindingLookup& BindingLookup, EStateTreeNodeFormatting Formatting /*= EStateTreeNodeFormatting::Text*/) const
 {
-	return LOCTEXT("StateTreeConditionGrounded", "<b>Is Character Grounded</b>");
+	return FText::FromString("<b>Is Character Grounded</b>");
 }
 #endif // WITH_EDITOR
 
@@ -72,7 +70,7 @@ void FStateTreeComboAttackTask::ExitState(FStateTreeExecutionContext& Context, c
 #if WITH_EDITOR
 FText FStateTreeComboAttackTask::GetDescription(const FGuid& ID, FStateTreeDataView InstanceDataView, const IStateTreeBindingLookup& BindingLookup, EStateTreeNodeFormatting Formatting /*= EStateTreeNodeFormatting::Text*/) const
 {
-	return LOCTEXT("StateTreeTaskComboAttackDescription", "<b>Do Combo Attack</b>");
+	return FText::FromString("<b>Do Combo Attack</b>");
 }
 #endif // WITH_EDITOR
 
@@ -117,7 +115,7 @@ void FStateTreeChargedAttackTask::ExitState(FStateTreeExecutionContext& Context,
 #if WITH_EDITOR
 FText FStateTreeChargedAttackTask::GetDescription(const FGuid& ID, FStateTreeDataView InstanceDataView, const IStateTreeBindingLookup& BindingLookup, EStateTreeNodeFormatting Formatting /*= EStateTreeNodeFormatting::Text*/) const
 {
-	return LOCTEXT("StateTreeTaskChargedAttackDescription", "<b>Do Charged Attack</b>");
+	return FText::FromString("<b>Do Charged Attack</b>");
 }
 #endif // WITH_EDITOR
 
@@ -159,7 +157,7 @@ void FStateTreeWaitForLandingTask::ExitState(FStateTreeExecutionContext& Context
 #if WITH_EDITOR
 FText FStateTreeWaitForLandingTask::GetDescription(const FGuid& ID, FStateTreeDataView InstanceDataView, const IStateTreeBindingLookup& BindingLookup, EStateTreeNodeFormatting Formatting /*= EStateTreeNodeFormatting::Text*/) const
 {
-	return LOCTEXT("StateTreeTaskChargedAttackDescription", "<b>Wait for Landing</b>");
+	return FText::FromString("<b>Wait for Landing</b>");
 }
 #endif // WITH_EDITOR
 
@@ -196,7 +194,7 @@ void FStateTreeFaceActorTask::ExitState(FStateTreeExecutionContext& Context, con
 #if WITH_EDITOR
 FText FStateTreeFaceActorTask::GetDescription(const FGuid& ID, FStateTreeDataView InstanceDataView, const IStateTreeBindingLookup& BindingLookup, EStateTreeNodeFormatting Formatting /*= EStateTreeNodeFormatting::Text*/) const
 {
-	return LOCTEXT("StateTreeTaskFaceActorDescription", "<b>Face Towards Actor</b>");
+	return FText::FromString("<b>Face Towards Actor</b>");
 }
 #endif // WITH_EDITOR
 
@@ -233,7 +231,7 @@ void FStateTreeFaceLocationTask::ExitState(FStateTreeExecutionContext& Context,
 #if WITH_EDITOR
 FText FStateTreeFaceLocationTask::GetDescription(const FGuid& ID, FStateTreeDataView InstanceDataView, const IStateTreeBindingLookup& BindingLookup, EStateTreeNodeFormatting Formatting /*= EStateTreeNodeFormatting::Text*/) const
 {
-	return LOCTEXT("StateTreeTaskFaceLocationDescription", "<b>Face Towards Location</b>");
+	return FText::FromString("<b>Face Towards Location</b>");
 }
 #endif // WITH_EDITOR
 
@@ -257,7 +255,7 @@ EStateTreeRunStatus FStateTreeSetCharacterSpeedTask::EnterState(FStateTreeExecut
 #if WITH_EDITOR
 FText FStateTreeSetCharacterSpeedTask::GetDescription(const FGuid& ID, FStateTreeDataView InstanceDataView, const IStateTreeBindingLookup& BindingLookup, EStateTreeNodeFormatting Formatting /*= EStateTreeNodeFormatting::Text*/) const
 {
-	return LOCTEXT("StateTreeTaskSetCharacterSpeedDescription", "<b>Set Character Speed</b>");
+	return FText::FromString("<b>Set Character Speed</b>");
 }
 #endif // WITH_EDITOR
 
@@ -287,6 +285,6 @@ EStateTreeRunStatus FStateTreeGetPlayerInfoTask::Tick(FStateTreeExecutionContext
 #if WITH_EDITOR
 FText FStateTreeGetPlayerInfoTask::GetDescription(const FGuid& ID, FStateTreeDataView InstanceDataView, const IStateTreeBindingLookup& BindingLookup, EStateTreeNodeFormatting Formatting /*= EStateTreeNodeFormatting::Text*/) const
 {
-	return LOCTEXT("StateTreeTaskGetPlayerInfoDescription", "<b>Get Player Info</b>");
+	return FText::FromString("<b>Get Player Info</b>");
 }
 #endif // WITH_EDITOR
\ No newline at end of file
diff --git a/Templates/TP_ThirdPerson/Source/TP_ThirdPerson/Variant_SideScrolling/AI/SideScrollingStateTreeUtility.cpp b/Templates/TP_ThirdPerson/Source/TP_ThirdPerson/Variant_SideScrolling/AI/SideScrollingStateTreeUtility.cpp
--- a/Templates/TP_ThirdPerson/Source/TP_ThirdPerson/Variant_SideScrolling/AI/SideScrollingStateTreeUtility.cpp
+++ b/Templates/TP_ThirdPerson/Source/TP_ThirdPerson/Variant_SideScrolling/AI/SideScrollingStateTreeUtility.cpp
@@ -7,9 +7,6 @@
 #include "AIController.h"
 #include "Kismet/GameplayStatics.h"
 
-
-#define LOCTEXT_NAMESPACE "SideScrollingTemplate"
-
 EStateTreeRunStatus FStateTreeGetPlayerTask::Tick(FStateTreeExecutionContext& Context, const float DeltaTime) const
 {
 	// get the instance data
@@ -30,6 +27,6 @@ EStateTreeRunStatus FStateTreeGetPlayerTask::Tick(FStateTreeExecutionContext& Co
 #if WITH_EDITOR
 FText FStateTreeGetPlayerTask::GetDescription(const FGuid& ID, FStateTreeDataView InstanceDataView, const IStateTreeBindingLookup& BindingLookup, EStateTreeNodeFormatting Formatting /*= EStateTreeNodeFormatting::Text*/) const
 {
-	return LOCTEXT("StateTreeTaskGetPlayer", "<b>Get Player</b>");
+	return FText::FromString("<b>Get Player</b>");
 }
 #endif // WITH_EDITOR
\ No newline at end of file
