Template bug fixes
TP_ThirdPerson - Fixed a runtime crash due to incorrectly created components on some classes.
TP_FirstPerson - Fixed broken pistol and material asset references on Shooter variant.
#rnx

diff --git a/Templates/TP_ThirdPerson/Source/TP_ThirdPerson/Variant_Combat/CombatDummy.cpp b/Templates/TP_ThirdPerson/Source/TP_ThirdPerson/Variant_Combat/CombatDummy.cpp
--- a/Templates/TP_ThirdPerson/Source/TP_ThirdPerson/Variant_Combat/CombatDummy.cpp
+++ b/Templates/TP_ThirdPerson/Source/TP_ThirdPerson/Variant_Combat/CombatDummy.cpp
@@ -11,7 +11,8 @@ ACombatDummy::ACombatDummy()
  	PrimaryActorTick.bCanEverTick = true;
 
 	// create the root
-	RootComponent = Root = CreateDefaultSubobject<USceneComponent>(TEXT("Root"));
+	Root = CreateDefaultSubobject<USceneComponent>(TEXT("Root"));
+	SetRootComponent(Root);
 
 	// create the base plate
 	BasePlate = CreateDefaultSubobject<UStaticMeshComponent>(TEXT("Base Plate"));
@@ -24,7 +25,7 @@ ACombatDummy::ACombatDummy()
 	Dummy->SetSimulatePhysics(true);
 
 	// create the physics constraint
-	PhysicsConstraint = CreateEditorOnlyDefaultSubobject<UPhysicsConstraintComponent>(TEXT("Physics Constraint"));
+	PhysicsConstraint = CreateDefaultSubobject<UPhysicsConstraintComponent>(TEXT("Physics Constraint"));
 	PhysicsConstraint->SetupAttachment(RootComponent);
 
 	PhysicsConstraint->SetConstrainedComponents(BasePlate, NAME_None, Dummy, NAME_None);
diff --git a/Templates/TP_ThirdPerson/Source/TP_ThirdPerson/Variant_SideScrolling/Gameplay/SideScrollingJumpPad.cpp b/Templates/TP_ThirdPerson/Source/TP_ThirdPerson/Variant_SideScrolling/Gameplay/SideScrollingJumpPad.cpp
--- a/Templates/TP_ThirdPerson/Source/TP_ThirdPerson/Variant_SideScrolling/Gameplay/SideScrollingJumpPad.cpp
+++ b/Templates/TP_ThirdPerson/Source/TP_ThirdPerson/Variant_SideScrolling/Gameplay/SideScrollingJumpPad.cpp
@@ -15,11 +15,11 @@ ASideScrollingJumpPad::ASideScrollingJumpPad()
 	RootComponent = CreateDefaultSubobject<USceneComponent>(TEXT("Root"));
 
 	// create the bounding box
-	Box = CreateEditorOnlyDefaultSubobject<UBoxComponent>(TEXT("Box"));
+	Box = CreateDefaultSubobject<UBoxComponent>(TEXT("Box"));
 	Box->SetupAttachment(RootComponent);
 
 	// configure the bounding box
-	Box->SetBoxExtent(FVector(115.0f, 90.0f, 20.0f));
+	Box->SetBoxExtent(FVector(115.0f, 90.0f, 20.0f), false);
 	Box->SetRelativeLocation(FVector(0.0f, 0.0f, 16.0f));
 
 	Box->SetCollisionObjectType(ECC_WorldDynamic);
