Templates
- TP_ThirdPerson, BP: Fixed missing collision on some level platforms that were causing the player to sink slightly on the floor.
- TP_ThirdPerson: Matched SideScrolling camera bounds with BP version.
- TP_ThirdPerson: Matched combat dummy physics reaction with BP version.
- TP_ThirdPerson, BP: Implemented Coyote Time jumps on Platforming and SideScrolling variants.
- TP_FirstPerson, BP: Fixed an issue that prevented press and hold fire with a gamepad on Shooter variants.
- TP_FirstPerson: Restored Shooter NPC life to 100 from 1

#rnx

diff --git a/Templates/TP_ThirdPerson/Source/TP_ThirdPerson/Variant_Platforming/PlatformingCharacter.cpp b/Templates/TP_ThirdPerson/Source/TP_ThirdPerson/Variant_Platforming/PlatformingCharacter.cpp
--- a/Templates/TP_ThirdPerson/Source/TP_ThirdPerson/Variant_Platforming/PlatformingCharacter.cpp
+++ b/Templates/TP_ThirdPerson/Source/TP_ThirdPerson/Variant_Platforming/PlatformingCharacter.cpp
@@ -29,8 +29,8 @@ APlatformingCharacter::APlatformingCharacter()
 	// enable press and hold jump
 	JumpMaxHoldTime = 0.4f;
 
-	// set the jump max count to 2 so we can double jump
-	JumpMaxCount = 2;
+	// set the jump max count to 3 so we can double jump and check for coyote time jumps
+	JumpMaxCount = 3;
 
 	// Set size for collision capsule
 	GetCapsuleComponent()->InitCapsuleSize(35.0f, 90.0f);
@@ -148,17 +148,35 @@ void APlatformingCharacter::MultiJump()
 			// no wall jump, try a double jump next
 			else
 			{
-				// only double jump once while we're in the air
-				if (!bHasDoubleJumped)
+				// are we still within coyote time frames?
+				if (GetWorld()->GetTimeSeconds() - LastFallTime < MaxCoyoteTime)
 				{
-					bHasDoubleJumped = true;
+					UE_LOG(LogTemp, Warning, TEXT("Coyote Jump"));
 
-					// use the built-in CMC functionality to do the double jump
+					// use the built-in CMC functionality to do the jump
 					Jump();
 
 					// enable the jump trail
 					SetJumpTrailState(true);
+
+				// no coyote time jump
+				} else {
+
+					// only double jump once while we're in the air
+					if (!bHasDoubleJumped)
+					{
+						bHasDoubleJumped = true;
+
+						// use the built-in CMC functionality to do the double jump
+						Jump();
+
+						// enable the jump trail
+						SetJumpTrailState(true);
+					}
+
 				}
+
+				
 			}
 		}
 
@@ -337,3 +355,15 @@ void APlatformingCharacter::Landed(const FHitResult& Hit)
 	SetJumpTrailState(false);
 }
 
+void APlatformingCharacter::OnMovementModeChanged(EMovementMode PrevMovementMode, uint8 PreviousCustomMode /*= 0*/)
+{
+	Super::OnMovementModeChanged(PrevMovementMode, PreviousCustomMode);
+
+	// are we falling?
+	if (GetCharacterMovement()->MovementMode == EMovementMode::MOVE_Falling)
+	{
+		// save the game time when we started falling, so we can check it later for coyote time jumps
+		LastFallTime = GetWorld()->GetTimeSeconds();
+	}
+}
+
diff --git a/Templates/TP_ThirdPerson/Source/TP_ThirdPerson/Variant_Platforming/PlatformingCharacter.h b/Templates/TP_ThirdPerson/Source/TP_ThirdPerson/Variant_Platforming/PlatformingCharacter.h
--- a/Templates/TP_ThirdPerson/Source/TP_ThirdPerson/Variant_Platforming/PlatformingCharacter.h
+++ b/Templates/TP_ThirdPerson/Source/TP_ThirdPerson/Variant_Platforming/PlatformingCharacter.h
@@ -136,6 +136,9 @@ public:
 	/** Handle landings to reset dash and advanced jump state */
 	virtual void Landed(const FHitResult& Hit) override;
 
+	/** Handle movement mode changes to keep track of coyote time jumps */
+	virtual void OnMovementModeChanged(EMovementMode PrevMovementMode, uint8 PreviousCustomMode = 0) override;
+
 protected:
 
 	/** movement state flag bits, packed into a uint8 for memory efficiency */
@@ -174,6 +177,13 @@ protected:
 	UPROPERTY(EditAnywhere, Category="Dash")
 	UAnimMontage* DashMontage;
 
+	/** Last recorded time when this character started falling */
+	float LastFallTime = 0.0f;
+
+	/** Max amount of time that can pass since we started falling when we allow a regular jump */
+	UPROPERTY(EditAnywhere, Category="Coyote Time", meta = (ClampMin = 0, ClampMax = 5, Units = "s"))
+	float MaxCoyoteTime = 0.16f;
+
 public:
 	/** Returns CameraBoom subobject **/
 	FORCEINLINE class USpringArmComponent* GetCameraBoom() const { return CameraBoom; }
diff --git a/Templates/TP_ThirdPerson/Source/TP_ThirdPerson/Variant_SideScrolling/SideScrollingCameraManager.h b/Templates/TP_ThirdPerson/Source/TP_ThirdPerson/Variant_SideScrolling/SideScrollingCameraManager.h
--- a/Templates/TP_ThirdPerson/Source/TP_ThirdPerson/Variant_SideScrolling/SideScrollingCameraManager.h
+++ b/Templates/TP_ThirdPerson/Source/TP_ThirdPerson/Variant_SideScrolling/SideScrollingCameraManager.h
@@ -30,11 +30,11 @@ public:
 	float CameraZOffset = 100.0f;
 
 	/** Minimum camera scrolling bounds in world space */
-	UPROPERTY(EditAnywhere, Category="Side Scrolling Camera", meta=(ClampMin=0, ClampMax=10000, Units="cm"))
+	UPROPERTY(EditAnywhere, Category="Side Scrolling Camera", meta=(ClampMin=-100000, ClampMax=100000, Units="cm"))
 	float CameraXMinBounds = -400.0f;
 
 	/** Maximum camera scrolling bounds in world space */
-	UPROPERTY(EditAnywhere, Category="Side Scrolling Camera", meta=(ClampMin=0, ClampMax=10000, Units="cm"))
+	UPROPERTY(EditAnywhere, Category="Side Scrolling Camera", meta=(ClampMin=-100000, ClampMax=100000, Units="cm"))
 	float CameraXMaxBounds = 10000.0f;
 
 protected:
diff --git a/Templates/TP_ThirdPerson/Source/TP_ThirdPerson/Variant_SideScrolling/SideScrollingCharacter.cpp b/Templates/TP_ThirdPerson/Source/TP_ThirdPerson/Variant_SideScrolling/SideScrollingCharacter.cpp
--- a/Templates/TP_ThirdPerson/Source/TP_ThirdPerson/Variant_SideScrolling/SideScrollingCharacter.cpp
+++ b/Templates/TP_ThirdPerson/Source/TP_ThirdPerson/Variant_SideScrolling/SideScrollingCharacter.cpp
@@ -55,8 +55,8 @@ ASideScrollingCharacter::ASideScrollingCharacter()
 	GetCharacterMovement()->SetPlaneConstraintNormal(FVector(0.0f, 1.0f, 0.0f));
 	GetCharacterMovement()->bConstrainToPlane = true;
 
-	// enable double jump
-	JumpMaxCount = 2;
+	// enable double jump and coyote time
+	JumpMaxCount = 3;
 }
 
 void ASideScrollingCharacter::EndPlay(EEndPlayReason::Type EndPlayReason)
@@ -121,6 +121,18 @@ void ASideScrollingCharacter::Landed(const FHitResult& Hit)
 	bHasDoubleJumped = false;
 }
 
+void ASideScrollingCharacter::OnMovementModeChanged(EMovementMode PrevMovementMode, uint8 PreviousCustomMode /*= 0*/)
+{
+	Super::OnMovementModeChanged(PrevMovementMode, PreviousCustomMode);
+
+	// are we falling?
+	if (GetCharacterMovement()->MovementMode == EMovementMode::MOVE_Falling)
+	{
+		// save the game time when we started falling, so we can check it later for coyote time jumps
+		LastFallTime = GetWorld()->GetTimeSeconds();
+	}
+}
+
 void ASideScrollingCharacter::Move(const FInputActionValue& Value)
 {
 	FVector2D MoveVector = Value.Get<FVector2D>();
@@ -259,17 +271,31 @@ void ASideScrollingCharacter::MultiJump()
 		}
 	}
 
+
+
 	// test for double jump only if we haven't already tested for wall jump
 	if (!bHasWallJumped)
 	{
-		// The movement component handles double jump but we still need to manage the flag for animation
-		if (!bHasDoubleJumped)
+		// are we still within coyote time frames?
+		if (GetWorld()->GetTimeSeconds() - LastFallTime < MaxCoyoteTime)
 		{
-			// raise the double jump flag
-			bHasDoubleJumped = true;
+			UE_LOG(LogTemp, Warning, TEXT("Coyote Jump"));
 
-			// let the CMC handle jump
+			// use the built-in CMC functionality to do the jump
 			Jump();
+
+		// no coyote time jump
+		} else {
+		
+			// The movement component handles double jump but we still need to manage the flag for animation
+			if (!bHasDoubleJumped)
+			{
+				// raise the double jump flag
+				bHasDoubleJumped = true;
+
+				// let the CMC handle jump
+				Jump();
+			}
 		}
 	}
 }
diff --git a/Templates/TP_ThirdPerson/Source/TP_ThirdPerson/Variant_SideScrolling/SideScrollingCharacter.h b/Templates/TP_ThirdPerson/Source/TP_ThirdPerson/Variant_SideScrolling/SideScrollingCharacter.h
--- a/Templates/TP_ThirdPerson/Source/TP_ThirdPerson/Variant_SideScrolling/SideScrollingCharacter.h
+++ b/Templates/TP_ThirdPerson/Source/TP_ThirdPerson/Variant_SideScrolling/SideScrollingCharacter.h
@@ -40,38 +40,45 @@ protected:
 	UPROPERTY(EditAnywhere, Category="Input")
 	UInputAction* InteractAction;
 
+	/** Impulse to manually push physics objects while we're in midair */
+	UPROPERTY(EditAnywhere, Category="Side Scrolling|Jump")
+	float JumpPushImpulse = 600.0f;
+
 	/** Max distance that interactive objects can be triggered */
-	UPROPERTY(EditAnywhere, Category="Side Scrolling")
+	UPROPERTY(EditAnywhere, Category="Side Scrolling|Interaction")
 	float InteractionRadius = 200.0f;
 
 	/** Time to disable input after a wall jump to preserve momentum */
-	UPROPERTY(EditAnywhere, Category="Side Scrolling")
+	UPROPERTY(EditAnywhere, Category="Side Scrolling|Wall Jump")
 	float DelayBetweenWallJumps = 0.3f;
 
-	/** Impulse to manually push physics objects while we're in midair */
-	UPROPERTY(EditAnywhere, Category="Side Scrolling")
-	float JumpPushImpulse = 600.0f;
-
 	/** Distance to trace ahead of the character for wall jumps */
-	UPROPERTY(EditAnywhere, Category="Side Scrolling")
+	UPROPERTY(EditAnywhere, Category="Side Scrolling|Wall Jump")
 	float WallJumpTraceDistance = 50.0f;
 
 	/** Horizontal impulse to apply to the character during wall jumps */
-	UPROPERTY(EditAnywhere, Category="Side Scrolling")
+	UPROPERTY(EditAnywhere, Category="Side Scrolling|Wall Jump")
 	float WallJumpHorizontalImpulse = 500.0f;
 
 	/** Multiplies the jump Z velocity for wall jumps. */
-	UPROPERTY(EditAnywhere, Category="Side Scrolling")
+	UPROPERTY(EditAnywhere, Category="Side Scrolling|Wall Jump")
 	float WallJumpVerticalMultiplier = 1.4f;
 
 	/** Collision object type to use for soft collision traces (dropping down floors) */
-	UPROPERTY(EditAnywhere, Category="Side Scrolling")
+	UPROPERTY(EditAnywhere, Category="Side Scrolling|Soft Platforms")
 	TEnumAsByte<ECollisionChannel> SoftCollisionObjectType;
 
 	/** Distance to trace down during soft collision checks */
-	UPROPERTY(EditAnywhere, Category="Side Scrolling")
+	UPROPERTY(EditAnywhere, Category="Side Scrolling|Soft Platforms")
 	float SoftCollisionTraceDistance = 1000.0f;
 
+	/** Last recorded time when this character started falling */
+	float LastFallTime = 0.0f;
+
+	/** Max amount of time that can pass since we started falling when we allow a regular jump */
+	UPROPERTY(EditAnywhere, Category="Side Scrolling|Coyote Time", meta = (ClampMin = 0, ClampMax = 5, Units = "s"))
+	float MaxCoyoteTime = 0.16f;
+
 	/** Wall jump lockout timer */
 	FTimerHandle WallJumpTimer;
 
@@ -109,6 +116,9 @@ protected:
 	/** Landing handling */
 	virtual void Landed(const FHitResult& Hit) override;
 
+	/** Handle movement mode changes to keep track of coyote time jumps */
+	virtual void OnMovementModeChanged(EMovementMode PrevMovementMode, uint8 PreviousCustomMode = 0) override;
+
 protected:
 
 	/** Called for movement input */
