Templates - TP_FirstPerson, TP_ThirdPerson
- Porting fix to incorrect rvalue access to statetree instance data from Main stream.
- Fixed uninitialized struct properties in State Tree Utilities
#rnx

diff --git a/Templates/TP_FirstPerson/Source/TP_FirstPerson/Variant_Shooter/AI/ShooterStateTreeUtility.h b/Templates/TP_FirstPerson/Source/TP_FirstPerson/Variant_Shooter/AI/ShooterStateTreeUtility.h
--- a/Templates/TP_FirstPerson/Source/TP_FirstPerson/Variant_Shooter/AI/ShooterStateTreeUtility.h
+++ b/Templates/TP_FirstPerson/Source/TP_FirstPerson/Variant_Shooter/AI/ShooterStateTreeUtility.h
@@ -125,7 +125,7 @@ struct FStateTreeFaceLocationInstanceData
 
 	/** Location that will be faced towards */
 	UPROPERTY(EditAnywhere, Category = Parameter)
-	FVector FaceLocation;
+	FVector FaceLocation = FVector::ZeroVector;
 };
 
 /**
@@ -260,7 +260,7 @@ struct FStateTreeSenseEnemiesInstanceData
 
 	/** Sensed location to investigate */
 	UPROPERTY(EditAnywhere, Category = Output)
-	FVector InvestigateLocation;
+	FVector InvestigateLocation = FVector::ZeroVector;
 
 	/** True if a target was successfully sensed */
 	UPROPERTY(EditAnywhere, Category = Output)
diff --git a/Templates/TP_ThirdPerson/Source/TP_ThirdPerson/Variant_Combat/AI/CombatStateTreeUtility.h b/Templates/TP_ThirdPerson/Source/TP_ThirdPerson/Variant_Combat/AI/CombatStateTreeUtility.h
--- a/Templates/TP_ThirdPerson/Source/TP_ThirdPerson/Variant_Combat/AI/CombatStateTreeUtility.h
+++ b/Templates/TP_ThirdPerson/Source/TP_ThirdPerson/Variant_Combat/AI/CombatStateTreeUtility.h
@@ -198,7 +198,7 @@ struct FStateTreeFaceLocationInstanceData
 
 	/** Location that will be faced towards */
 	UPROPERTY(EditAnywhere, Category = Parameter)
-	FVector FaceLocation;
+	FVector FaceLocation = FVector::ZeroVector;
 };
 
 /**
@@ -283,11 +283,11 @@ struct FStateTreeGetPlayerInfoInstanceData
 
 	/** Last known location for the target */
 	UPROPERTY(VisibleAnywhere)
-	FVector TargetPlayerLocation;
+	FVector TargetPlayerLocation = FVector::ZeroVector;
 
 	/** Distance to the target */
 	UPROPERTY(VisibleAnywhere)
-	float DistanceToTarget;
+	float DistanceToTarget = 0.0f;
 };
 
 /**
