TP_ThirdPerson, TP_ThirdPersonBP
- Removed unused Dash input from input mapping context on combat variant
- Implemented camera shoulder toggle
- Implemented enemy evasion on player attacks

#rnx
#virtualized

diff --git a/Templates/TP_ThirdPerson/Source/TP_ThirdPerson/Variant_Combat/AI/CombatEnemy.cpp b/Templates/TP_ThirdPerson/Source/TP_ThirdPerson/Variant_Combat/AI/CombatEnemy.cpp
--- a/Templates/TP_ThirdPerson/Source/TP_ThirdPerson/Variant_Combat/AI/CombatEnemy.cpp
+++ b/Templates/TP_ThirdPerson/Source/TP_ThirdPerson/Variant_Combat/AI/CombatEnemy.cpp
@@ -113,6 +113,16 @@ void ACombatEnemy::AttackMontageEnded(UAnimMontage* Montage, bool bInterrupted)
 	OnAttackCompleted.ExecuteIfBound();
 }
 
+const FVector& ACombatEnemy::GetLastDangerLocation() const
+{
+	return LastDangerLocation;
+}
+
+float ACombatEnemy::GetLastDangerTime() const
+{
+	return LastDangerTime;
+}
+
 void ACombatEnemy::DoAttackTrace(FName DamageSourceBone)
 {
 	// sweep for objects in front of the character to be hit by the attack
@@ -245,6 +255,17 @@ void ACombatEnemy::ApplyHealing(float Healing, AActor* Healer)
 	// stub
 }
 
+void ACombatEnemy::NotifyDanger(const FVector& DangerLocation, AActor* DangerSource)
+{
+	// ensure we're being attacked by the player
+	if (DangerSource && DangerSource->ActorHasTag(FName("Player")))
+	{
+		// save the danger location and game time
+		LastDangerLocation = DangerLocation;
+		LastDangerTime = GetWorld()->GetTimeSeconds();
+	}
+}
+
 void ACombatEnemy::RemoveFromLevel()
 {
 	// destroy this actor
diff --git a/Templates/TP_ThirdPerson/Source/TP_ThirdPerson/Variant_Combat/AI/CombatEnemy.h b/Templates/TP_ThirdPerson/Source/TP_ThirdPerson/Variant_Combat/AI/CombatEnemy.h
--- a/Templates/TP_ThirdPerson/Source/TP_ThirdPerson/Variant_Combat/AI/CombatEnemy.h
+++ b/Templates/TP_ThirdPerson/Source/TP_ThirdPerson/Variant_Combat/AI/CombatEnemy.h
@@ -136,6 +136,12 @@ protected:
 	/** Attack montage ended delegate */
 	FOnMontageEnded OnAttackMontageEnded;
 
+	/** Last recorded location we're being attacked from */
+	FVector LastDangerLocation = FVector::ZeroVector;
+
+	/** Last recorded game time we were attacked */
+	float LastDangerTime = -1000.0f;
+
 public:
 	/** Attack completed internal delegate to notify StateTree tasks */
 	FOnEnemyAttackCompleted OnAttackCompleted;
@@ -158,6 +164,12 @@ public:
 	/** Called from a delegate when the attack montage ends */
 	void AttackMontageEnded(UAnimMontage* Montage, bool bInterrupted);
 
+	/** Returns the last recorded location we were attacked from */
+	const FVector& GetLastDangerLocation() const;
+
+	/** Returns the last game time we were attacked */
+	float GetLastDangerTime() const;
+
 public:
 
 	// ~begin ICombatAttacker interface
@@ -186,6 +198,9 @@ public:
 	/** Handles healing events */
 	virtual void ApplyHealing(float Healing, AActor* Healer) override;
 
+	/** Allows the enemy to react to incoming attacks */
+	virtual void NotifyDanger(const FVector& DangerLocation, AActor* DangerSource) override;
+
 	// ~end ICombatDamageable interface
 
 protected:
diff --git a/Templates/TP_ThirdPerson/Source/TP_ThirdPerson/Variant_Combat/AI/CombatStateTreeUtility.cpp b/Templates/TP_ThirdPerson/Source/TP_ThirdPerson/Variant_Combat/AI/CombatStateTreeUtility.cpp
--- a/Templates/TP_ThirdPerson/Source/TP_ThirdPerson/Variant_Combat/AI/CombatStateTreeUtility.cpp
+++ b/Templates/TP_ThirdPerson/Source/TP_ThirdPerson/Variant_Combat/AI/CombatStateTreeUtility.cpp
@@ -30,6 +30,40 @@ FText FStateTreeCharacterGroundedCondition::GetDescription(const FGuid& ID, FSta
 
 ////////////////////////////////////////////////////////////////////
 
+bool FStateTreeIsInDangerCondition::TestCondition(FStateTreeExecutionContext& Context) const
+{
+	const FInstanceDataType& InstanceData = Context.GetInstanceData(*this);
+
+	// ensure we have a valid enemy character
+	if (InstanceData.Character)
+	{
+		// is the last detected danger event within the reaction threshold?
+		const float ReactionDelta = InstanceData.Character->GetWorld()->GetTimeSeconds() - InstanceData.Character->GetLastDangerTime();
+
+		if (ReactionDelta < InstanceData.MaxReactionTime && ReactionDelta > InstanceData.MinReactionTime)
+		{
+			// do a dot product check to determine if the danger location is within the character's detection cone
+			const FVector DangerDir = (InstanceData.Character->GetLastDangerLocation() - InstanceData.Character->GetActorLocation()).GetSafeNormal2D();
+
+			const float DangerDot = FVector::DotProduct(DangerDir, InstanceData.Character->GetActorForwardVector());
+			const float ConeAngleCos = FMath::Cos(FMath::DegreesToRadians(InstanceData.DangerSightConeAngle));
+
+			return DangerDot > ConeAngleCos;
+		}
+	}
+
+	return false;
+}
+
+#if WITH_EDITOR
+FText FStateTreeIsInDangerCondition::GetDescription(const FGuid& ID, FStateTreeDataView InstanceDataView, const IStateTreeBindingLookup& BindingLookup, EStateTreeNodeFormatting Formatting /*= EStateTreeNodeFormatting::Text*/) const
+{
+	return FText::FromString("<b>Is Character In Danger</b>");
+}
+#endif // WITH_EDITOR
+
+////////////////////////////////////////////////////////////////////
+
 EStateTreeRunStatus FStateTreeComboAttackTask::EnterState(FStateTreeExecutionContext& Context, const FStateTreeTransitionResult& Transition) const
 {
 	// have we transitioned from another state?
diff --git a/Templates/TP_ThirdPerson/Source/TP_ThirdPerson/Variant_Combat/AI/CombatStateTreeUtility.h b/Templates/TP_ThirdPerson/Source/TP_ThirdPerson/Variant_Combat/AI/CombatStateTreeUtility.h
--- a/Templates/TP_ThirdPerson/Source/TP_ThirdPerson/Variant_Combat/AI/CombatStateTreeUtility.h
+++ b/Templates/TP_ThirdPerson/Source/TP_ThirdPerson/Variant_Combat/AI/CombatStateTreeUtility.h
@@ -58,6 +58,60 @@ struct FStateTreeCharacterGroundedCondition : public FStateTreeConditionCommonBa
 
 ////////////////////////////////////////////////////////////////////
 
+/**
+ *  Instance data struct for the FStateTreeIsInDangerCondition condition
+ */
+USTRUCT()
+struct FStateTreeIsInDangerConditionInstanceData
+{
+	GENERATED_BODY()
+	
+	/** Character to check grounded status on */
+	UPROPERTY(EditAnywhere, Category = "Context")
+	ACombatEnemy* Character;
+
+	/** Minimum time to wait before reacting to the danger event */
+	UPROPERTY(EditAnywhere, Category = "Parameters", meta = (Units = "s"))
+	float MinReactionTime = 0.35f;
+
+	/** Maximum time to wait before ignoring the danger event */
+	UPROPERTY(EditAnywhere, Category = "Parameters", meta = (Units = "s"))
+	float MaxReactionTime = 0.75f;
+
+	/** Line of sight half angle for detecting incoming danger, in degrees*/
+	UPROPERTY(EditAnywhere, Category = "Parameters", meta = (Units = "degrees"))
+	float DangerSightConeAngle = 120.0f;
+};
+STATETREE_POD_INSTANCEDATA(FStateTreeIsInDangerConditionInstanceData);
+
+/**
+ *  StateTree condition to check if the character is about to be hit by an attack
+ */
+USTRUCT(DisplayName = "Character is in Danger")
+struct FStateTreeIsInDangerCondition : public FStateTreeConditionCommonBase
+{
+	GENERATED_BODY()
+
+	/** Set the instance data type */
+	using FInstanceDataType = FStateTreeIsInDangerConditionInstanceData;
+	virtual const UStruct* GetInstanceDataType() const override { return FInstanceDataType::StaticStruct(); }
+
+	/** Default constructor */
+	FStateTreeIsInDangerCondition() = default;
+	
+	/** Tests the StateTree condition */
+	virtual bool TestCondition(FStateTreeExecutionContext& Context) const override;
+
+#if WITH_EDITOR
+
+	/** Provides the description string */
+	virtual FText GetDescription(const FGuid& ID, FStateTreeDataView InstanceDataView, const IStateTreeBindingLookup& BindingLookup, EStateTreeNodeFormatting Formatting = EStateTreeNodeFormatting::Text) const override;
+#endif
+
+};
+
+////////////////////////////////////////////////////////////////////
+
 /**
  *  Instance data struct for the Combat StateTree tasks
  */
diff --git a/Templates/TP_ThirdPerson/Source/TP_ThirdPerson/Variant_Combat/AI/EnvQueryContext_Danger.cpp b/Templates/TP_ThirdPerson/Source/TP_ThirdPerson/Variant_Combat/AI/EnvQueryContext_Danger.cpp
new file mode 100644
--- /dev/null
+++ b/Templates/TP_ThirdPerson/Source/TP_ThirdPerson/Variant_Combat/AI/EnvQueryContext_Danger.cpp
@@ -0,0 +1,17 @@
+// Copyright Epic Games, Inc. All Rights Reserved.
+
+
+#include "Variant_Combat/AI/EnvQueryContext_Danger.h"
+#include "Variant_Combat/AI/CombatEnemy.h"
+#include "EnvironmentQuery/EnvQueryTypes.h"
+#include "EnvironmentQuery/Items/EnvQueryItemType_Point.h"
+
+void UEnvQueryContext_Danger::ProvideContext(FEnvQueryInstance& QueryInstance, FEnvQueryContextData& ContextData) const
+{
+	// get the querying enemy
+	if (ACombatEnemy* QuerierActor = Cast<ACombatEnemy>(QueryInstance.Owner.Get()))
+	{
+		// add the last recorded danger location to the context
+		UEnvQueryItemType_Point::SetContextHelper(ContextData, QuerierActor->GetLastDangerLocation());
+	}
+}
\ No newline at end of file
diff --git a/Templates/TP_ThirdPerson/Source/TP_ThirdPerson/Variant_Combat/AI/EnvQueryContext_Danger.h b/Templates/TP_ThirdPerson/Source/TP_ThirdPerson/Variant_Combat/AI/EnvQueryContext_Danger.h
new file mode 100644
--- /dev/null
+++ b/Templates/TP_ThirdPerson/Source/TP_ThirdPerson/Variant_Combat/AI/EnvQueryContext_Danger.h
@@ -0,0 +1,23 @@
+// Copyright Epic Games, Inc. All Rights Reserved.
+
+#pragma once
+
+#include "CoreMinimal.h"
+#include "EnvironmentQuery/EnvQueryContext.h"
+#include "EnvQueryContext_Danger.generated.h"
+
+/**
+ *  UEnvQueryContext_Danger
+ *  Returns the enemy character's last known danger location
+ */
+UCLASS()
+class TP_THIRDPERSON_API UEnvQueryContext_Danger : public UEnvQueryContext
+{
+	GENERATED_BODY()
+	
+public:
+
+	/** Provides the context locations or actors for this EnvQuery */
+	virtual void ProvideContext(FEnvQueryInstance& QueryInstance, FEnvQueryContextData& ContextData) const override;
+
+};
diff --git a/Templates/TP_ThirdPerson/Source/TP_ThirdPerson/Variant_Combat/CombatCharacter.cpp b/Templates/TP_ThirdPerson/Source/TP_ThirdPerson/Variant_Combat/CombatCharacter.cpp
--- a/Templates/TP_ThirdPerson/Source/TP_ThirdPerson/Variant_Combat/CombatCharacter.cpp
+++ b/Templates/TP_ThirdPerson/Source/TP_ThirdPerson/Variant_Combat/CombatCharacter.cpp
@@ -86,6 +86,12 @@ void ACombatCharacter::ChargedAttackReleased()
 	DoChargedAttackEnd();
 }
 
+void ACombatCharacter::ToggleCamera()
+{
+	// call the BP hook
+	BP_ToggleCamera();
+}
+
 void ACombatCharacter::DoMove(float Right, float Forward)
 {
 	if (GetController() != nullptr)
@@ -181,6 +187,9 @@ void ACombatCharacter::ComboAttack()
 	// reset the combo count
 	ComboCount = 0;
 
+	// notify enemies they are about to be attacked
+	NotifyEnemiesOfAttack();
+
 	// play the attack montage
 	if (UAnimInstance* AnimInstance = GetMesh()->GetAnimInstance())
 	{
@@ -204,6 +213,9 @@ void ACombatCharacter::ChargedAttack()
 	// reset the charge loop flag
 	bHasLoopedChargedAttack = false;
 
+	// notify enemies they are about to be attacked
+	NotifyEnemiesOfAttack();
+
 	// play the charged attack montage
 	if (UAnimInstance* AnimInstance = GetMesh()->GetAnimInstance())
 	{
@@ -302,6 +314,9 @@ void ACombatCharacter::CheckCombo()
 			// do we still have a combo section to play?
 			if (ComboCount < ComboSectionNames.Num())
 			{
+				// notify enemies they are about to be attacked
+				NotifyEnemiesOfAttack();
+
 				// jump to the next combo section
 				if (UAnimInstance* AnimInstance = GetMesh()->GetAnimInstance())
 				{
@@ -324,6 +339,44 @@ void ACombatCharacter::CheckChargedAttack()
 	}
 }
 
+void ACombatCharacter::NotifyEnemiesOfAttack()
+{
+	// sweep for objects in front of the character to be hit by the attack
+	TArray<FHitResult> OutHits;
+
+	// start at the actor location, sweep forward
+	const FVector TraceStart = GetActorLocation();
+	const FVector TraceEnd = TraceStart + (GetActorForwardVector() * DangerTraceDistance);
+
+	// check for pawn object types only
+	FCollisionObjectQueryParams ObjectParams;
+	ObjectParams.AddObjectTypesToQuery(ECC_Pawn);
+
+	// use a sphere shape for the sweep
+	FCollisionShape CollisionShape;
+	CollisionShape.SetSphere(DangerTraceRadius);
+
+	// ignore self
+	FCollisionQueryParams QueryParams;
+	QueryParams.AddIgnoredActor(this);
+
+	if (GetWorld()->SweepMultiByObjectType(OutHits, TraceStart, TraceEnd, FQuat::Identity, ObjectParams, CollisionShape, QueryParams))
+	{
+		// iterate over each object hit
+		for (const FHitResult& CurrentHit : OutHits)
+		{
+			// check if we've hit a damageable actor
+			ICombatDamageable* Damageable = Cast<ICombatDamageable>(CurrentHit.GetActor());
+
+			if (Damageable)
+			{
+				// notify the enemy
+				Damageable->NotifyDanger(GetActorLocation(), this);
+			}
+		}
+	}
+}
+
 void ACombatCharacter::ApplyDamage(float Damage, AActor* DamageCauser, const FVector& DamageLocation, const FVector& DamageImpulse)
 {
 	// pass the damage event to the actor
@@ -372,6 +425,11 @@ void ACombatCharacter::ApplyHealing(float Healing, AActor* Healer)
 	// stub
 }
 
+void ACombatCharacter::NotifyDanger(const FVector& DangerLocation, AActor* DangerSource)
+{
+	// stub
+}
+
 void ACombatCharacter::RespawnCharacter()
 {
 	// destroy the character and let it be respawned by the Player Controller
@@ -470,6 +528,9 @@ void ACombatCharacter::SetupPlayerInputComponent(UInputComponent* PlayerInputCom
 		// Charged Attack
 		EnhancedInputComponent->BindAction(ChargedAttackAction, ETriggerEvent::Started, this, &ACombatCharacter::ChargedAttackPressed);
 		EnhancedInputComponent->BindAction(ChargedAttackAction, ETriggerEvent::Completed, this, &ACombatCharacter::ChargedAttackReleased);
+
+		// Camera Side Toggle
+		EnhancedInputComponent->BindAction(ToggleCameraAction, ETriggerEvent::Triggered, this, &ACombatCharacter::ToggleCamera);
 	}
 }
 
diff --git a/Templates/TP_ThirdPerson/Source/TP_ThirdPerson/Variant_Combat/CombatCharacter.h b/Templates/TP_ThirdPerson/Source/TP_ThirdPerson/Variant_Combat/CombatCharacter.h
--- a/Templates/TP_ThirdPerson/Source/TP_ThirdPerson/Variant_Combat/CombatCharacter.h
+++ b/Templates/TP_ThirdPerson/Source/TP_ThirdPerson/Variant_Combat/CombatCharacter.h
@@ -69,6 +69,10 @@ protected:
 	UPROPERTY(EditAnywhere, Category ="Input")
 	UInputAction* ChargedAttackAction;
 
+	/** Toggle Camera Side Input Action */
+	UPROPERTY(EditAnywhere, Category ="Input")
+	UInputAction* ToggleCameraAction;
+
 	/** Max amount of HP the character will have on respawn */
 	UPROPERTY(EditAnywhere, Category="Damage", meta = (ClampMin = 0, ClampMax = 100))
 	float MaxHP = 5.0f;
@@ -107,6 +111,14 @@ protected:
 	UPROPERTY(EditAnywhere, Category="Melee Attack|Trace", meta = (ClampMin = 0, ClampMax = 200, Units = "cm"))
 	float MeleeTraceRadius = 75.0f;
 
+	/** Distance ahead of the character that enemies will be notified of incoming attacks */
+	UPROPERTY(EditAnywhere, Category="Melee Attack|Trace", meta = (ClampMin = 0, ClampMax = 500, Units="cm"))
+	float DangerTraceDistance = 300.0f;
+
+	/** Radius of the sphere trace to notify enemies of incoming attacks */
+	UPROPERTY(EditAnywhere, Category="Melee Attack|Trace", meta = (ClampMin = 0, ClampMax = 200, Units = "cm"))
+	float DangerTraceRadius = 100.0f;
+
 	/** Amount of damage a melee attack will deal */
 	UPROPERTY(EditAnywhere, Category="Melee Attack|Damage", meta = (ClampMin = 0, ClampMax = 100))
 	float MeleeDamage = 1.0f;
@@ -195,6 +207,13 @@ protected:
 	/** Called for combo attack input released */
 	void ChargedAttackReleased();
 
+	/** Called for toggle camera side input */
+	void ToggleCamera();
+
+	/** BP hook to animate the camera side switch */
+	UFUNCTION(BlueprintImplementableEvent, Category="Combat")
+	void BP_ToggleCamera();
+
 public:
 
 	/** Handles move inputs from either controls or UI interfaces */
@@ -253,6 +272,9 @@ public:
 
 	// ~begin CombatDamageable interface
 
+	/** Notifies nearby enemies that an attack is coming so they can react */
+	void NotifyEnemiesOfAttack();
+
 	/** Handles damage and knockback events */
 	virtual void ApplyDamage(float Damage, AActor* DamageCauser, const FVector& DamageLocation, const FVector& DamageImpulse) override;
 
@@ -262,6 +284,9 @@ public:
 	/** Handles healing events */
 	virtual void ApplyHealing(float Healing, AActor* Healer) override;
 
+	/** Allows reaction to incoming attacks */
+	virtual void NotifyDanger(const FVector& DangerLocation, AActor* DangerSource) override;
+
 	// ~end CombatDamageable interface
 
 	/** Called from the respawn timer to destroy and re-create the character */
diff --git a/Templates/TP_ThirdPerson/Source/TP_ThirdPerson/Variant_Combat/Gameplay/CombatDamageableBox.cpp b/Templates/TP_ThirdPerson/Source/TP_ThirdPerson/Variant_Combat/Gameplay/CombatDamageableBox.cpp
--- a/Templates/TP_ThirdPerson/Source/TP_ThirdPerson/Variant_Combat/Gameplay/CombatDamageableBox.cpp
+++ b/Templates/TP_ThirdPerson/Source/TP_ThirdPerson/Variant_Combat/Gameplay/CombatDamageableBox.cpp
@@ -76,3 +76,8 @@ void ACombatDamageableBox::ApplyHealing(float Healing, AActor* Healer)
 	// stub
 }
 
+void ACombatDamageableBox::NotifyDanger(const FVector& DangerLocation, AActor* DangerSource)
+{
+	// stub
+}
+
diff --git a/Templates/TP_ThirdPerson/Source/TP_ThirdPerson/Variant_Combat/Gameplay/CombatDamageableBox.h b/Templates/TP_ThirdPerson/Source/TP_ThirdPerson/Variant_Combat/Gameplay/CombatDamageableBox.h
--- a/Templates/TP_ThirdPerson/Source/TP_ThirdPerson/Variant_Combat/Gameplay/CombatDamageableBox.h
+++ b/Templates/TP_ThirdPerson/Source/TP_ThirdPerson/Variant_Combat/Gameplay/CombatDamageableBox.h
@@ -64,5 +64,8 @@ public:
 	/** Handles healing events */
 	virtual void ApplyHealing(float Healing, AActor* Healer) override;
 
+	/** Allows reaction to incoming attacks */
+	virtual void NotifyDanger(const FVector& DangerLocation, AActor* DangerSource) override;
+
 	// ~End CombatDamageable interface
 };
diff --git a/Templates/TP_ThirdPerson/Source/TP_ThirdPerson/Variant_Combat/Gameplay/CombatDummy.cpp b/Templates/TP_ThirdPerson/Source/TP_ThirdPerson/Variant_Combat/Gameplay/CombatDummy.cpp
--- a/Templates/TP_ThirdPerson/Source/TP_ThirdPerson/Variant_Combat/Gameplay/CombatDummy.cpp
+++ b/Templates/TP_ThirdPerson/Source/TP_ThirdPerson/Variant_Combat/Gameplay/CombatDummy.cpp
@@ -48,4 +48,9 @@ void ACombatDummy::HandleDeath()
 void ACombatDummy::ApplyHealing(float Healing, AActor* Healer)
 {
 	// unused
-}
\ No newline at end of file
+}
+
+void ACombatDummy::NotifyDanger(const FVector& DangerLocation, AActor* DangerSource)
+{
+	// unused
+}
diff --git a/Templates/TP_ThirdPerson/Source/TP_ThirdPerson/Variant_Combat/Gameplay/CombatDummy.h b/Templates/TP_ThirdPerson/Source/TP_ThirdPerson/Variant_Combat/Gameplay/CombatDummy.h
--- a/Templates/TP_ThirdPerson/Source/TP_ThirdPerson/Variant_Combat/Gameplay/CombatDummy.h
+++ b/Templates/TP_ThirdPerson/Source/TP_ThirdPerson/Variant_Combat/Gameplay/CombatDummy.h
@@ -50,6 +50,9 @@ public:
 	/** Handles healing events */
 	virtual void ApplyHealing(float Healing, AActor* Healer) override;
 
+	/** Allows reaction to incoming attacks */
+	virtual void NotifyDanger(const FVector& DangerLocation, AActor* DangerSource) override;
+
 	// ~End CombatDamageable interface
 
 protected:
diff --git a/Templates/TP_ThirdPerson/Source/TP_ThirdPerson/Variant_Combat/Interfaces/CombatDamageable.h b/Templates/TP_ThirdPerson/Source/TP_ThirdPerson/Variant_Combat/Interfaces/CombatDamageable.h
--- a/Templates/TP_ThirdPerson/Source/TP_ThirdPerson/Variant_Combat/Interfaces/CombatDamageable.h
+++ b/Templates/TP_ThirdPerson/Source/TP_ThirdPerson/Variant_Combat/Interfaces/CombatDamageable.h
@@ -33,4 +33,8 @@ public:
 	/** Handles healing events */
 	UFUNCTION(BlueprintCallable, Category="Damageable")
 	virtual void ApplyHealing(float Healing, AActor* Healer) = 0;
+
+	/** Notifies the actor of impending danger such as an incoming hit, allowing it to react. */
+	UFUNCTION(BlueprintCallable, Category="Damageable")
+	virtual void NotifyDanger(const FVector& DangerLocation, AActor* DangerSource) = 0;
 };
