Unshelved from pending changelist '45972970':

TP_ThirdPerson
- Fixed a build error due to missing UWorld include in some source files.
#rnx

diff --git a/Templates/TP_ThirdPerson/Source/TP_ThirdPerson/Variant_Combat/AI/CombatEnemySpawner.cpp b/Templates/TP_ThirdPerson/Source/TP_ThirdPerson/Variant_Combat/AI/CombatEnemySpawner.cpp
--- a/Templates/TP_ThirdPerson/Source/TP_ThirdPerson/Variant_Combat/AI/CombatEnemySpawner.cpp
+++ b/Templates/TP_ThirdPerson/Source/TP_ThirdPerson/Variant_Combat/AI/CombatEnemySpawner.cpp
@@ -2,6 +2,7 @@
 
 
 #include "CombatEnemySpawner.h"
+#include "Engine/World.h"
 #include "Components/SceneComponent.h"
 #include "Components/CapsuleComponent.h"
 #include "Components/ArrowComponent.h"
