TP_ThirdPerson
- Added missing UWorld includes on source files.
#rnx

diff --git a/Templates/TP_ThirdPerson/Source/TP_ThirdPerson/Variant_Combat/AI/CombatStateTreeUtility.cpp b/Templates/TP_ThirdPerson/Source/TP_ThirdPerson/Variant_Combat/AI/CombatStateTreeUtility.cpp
--- a/Templates/TP_ThirdPerson/Source/TP_ThirdPerson/Variant_Combat/AI/CombatStateTreeUtility.cpp
+++ b/Templates/TP_ThirdPerson/Source/TP_ThirdPerson/Variant_Combat/AI/CombatStateTreeUtility.cpp
@@ -4,6 +4,7 @@
 #include "CombatStateTreeUtility.h"
 #include "StateTreeExecutionContext.h"
 #include "StateTreeExecutionTypes.h"
+#include "Engine/World.h"
 #include "GameFramework/Character.h"
 #include "GameFramework/CharacterMovementComponent.h"
 #include "AIController.h"
diff --git a/Templates/TP_ThirdPerson/Source/TP_ThirdPerson/Variant_Platforming/PlatformingCharacter.cpp b/Templates/TP_ThirdPerson/Source/TP_ThirdPerson/Variant_Platforming/PlatformingCharacter.cpp
--- a/Templates/TP_ThirdPerson/Source/TP_ThirdPerson/Variant_Platforming/PlatformingCharacter.cpp
+++ b/Templates/TP_ThirdPerson/Source/TP_ThirdPerson/Variant_Platforming/PlatformingCharacter.cpp
@@ -4,6 +4,7 @@
 #include "PlatformingCharacter.h"
 
 #include "Components/CapsuleComponent.h"
+#include "Engine/World.h"
 #include "GameFramework/CharacterMovementComponent.h"
 #include "GameFramework/SpringArmComponent.h"
 #include "Components/SkeletalMeshComponent.h"
diff --git a/Templates/TP_ThirdPerson/Source/TP_ThirdPerson/Variant_SideScrolling/AI/SideScrollingNPC.cpp b/Templates/TP_ThirdPerson/Source/TP_ThirdPerson/Variant_SideScrolling/AI/SideScrollingNPC.cpp
--- a/Templates/TP_ThirdPerson/Source/TP_ThirdPerson/Variant_SideScrolling/AI/SideScrollingNPC.cpp
+++ b/Templates/TP_ThirdPerson/Source/TP_ThirdPerson/Variant_SideScrolling/AI/SideScrollingNPC.cpp
@@ -2,6 +2,7 @@
 
 
 #include "SideScrollingNPC.h"
+#include "Engine/World.h"
 #include "GameFramework/CharacterMovementComponent.h"
 #include "TimerManager.h"
 
