Templates
- TP_FirstPerson: fixed a NPC spawner issue that caused spawners to stop working after the first NPC they spawned died.
- TP_FirstPerson: NPC spawners no longer spawn the first NPC if their spawn count is zero
- Fixed multiple errors, typos and ambiguities on C++ code and comments
- Fixed potential garbage collection vulnerabilities on UMG widget pointers
- Fixed bForceTouchControls on player controllers not being assigned correctly to a config file
#rnx

diff --git a/Templates/TP_FirstPerson/Source/TP_FirstPerson/TP_FirstPersonPlayerController.h b/Templates/TP_FirstPerson/Source/TP_FirstPerson/TP_FirstPersonPlayerController.h
--- a/Templates/TP_FirstPerson/Source/TP_FirstPerson/TP_FirstPersonPlayerController.h
+++ b/Templates/TP_FirstPerson/Source/TP_FirstPerson/TP_FirstPersonPlayerController.h
@@ -39,6 +39,7 @@ protected:
 	TSubclassOf<UUserWidget> MobileControlsWidgetClass;
 
 	/** Pointer to the mobile controls widget */
+	UPROPERTY()
 	TObjectPtr<UUserWidget> MobileControlsWidget;
 
 	/** If true, the player will use UMG touch controls even if not playing on mobile platforms */
diff --git a/Templates/TP_FirstPerson/Source/TP_FirstPerson/Variant_Horror/HorrorPlayerController.h b/Templates/TP_FirstPerson/Source/TP_FirstPerson/Variant_Horror/HorrorPlayerController.h
--- a/Templates/TP_FirstPerson/Source/TP_FirstPerson/Variant_Horror/HorrorPlayerController.h
+++ b/Templates/TP_FirstPerson/Source/TP_FirstPerson/Variant_Horror/HorrorPlayerController.h
@@ -26,6 +26,7 @@ protected:
 	TSubclassOf<UHorrorUI> HorrorUIClass;
 
 	/** Pointer to the UI widget */
+	UPROPERTY()
 	TObjectPtr<UHorrorUI> HorrorUI;
 
 public:
diff --git a/Templates/TP_FirstPerson/Source/TP_FirstPerson/Variant_Shooter/AI/ShooterNPC.cpp b/Templates/TP_FirstPerson/Source/TP_FirstPerson/Variant_Shooter/AI/ShooterNPC.cpp
--- a/Templates/TP_FirstPerson/Source/TP_FirstPerson/Variant_Shooter/AI/ShooterNPC.cpp
+++ b/Templates/TP_FirstPerson/Source/TP_FirstPerson/Variant_Shooter/AI/ShooterNPC.cpp
@@ -62,7 +62,7 @@ void AShooterNPC::AttachWeaponMeshes(AShooterWeapon* WeaponToAttach)
 
 	// attach the weapon meshes
 	WeaponToAttach->GetFirstPersonMesh()->AttachToComponent(GetFirstPersonMesh(), AttachmentRule, FirstPersonWeaponSocket);
-	WeaponToAttach->GetThirdPersonMesh()->AttachToComponent(GetMesh(), AttachmentRule, FirstPersonWeaponSocket);
+	WeaponToAttach->GetThirdPersonMesh()->AttachToComponent(GetMesh(), AttachmentRule, ThirdPersonWeaponSocket);
 }
 
 void AShooterNPC::PlayFiringMontage(UAnimMontage* Montage)
@@ -162,6 +162,9 @@ void AShooterNPC::Die()
 	// grant the death tag to the character
 	Tags.Add(DeathTag);
 
+	// call the delegate
+	OnPawnDeath.Broadcast();
+
 	// increment the team score
 	if (AShooterGameMode* GM = Cast<AShooterGameMode>(GetWorld()->GetAuthGameMode()))
 	{
diff --git a/Templates/TP_FirstPerson/Source/TP_FirstPerson/Variant_Shooter/AI/ShooterNPCSpawner.cpp b/Templates/TP_FirstPerson/Source/TP_FirstPerson/Variant_Shooter/AI/ShooterNPCSpawner.cpp
--- a/Templates/TP_FirstPerson/Source/TP_FirstPerson/Variant_Shooter/AI/ShooterNPCSpawner.cpp
+++ b/Templates/TP_FirstPerson/Source/TP_FirstPerson/Variant_Shooter/AI/ShooterNPCSpawner.cpp
@@ -33,8 +33,8 @@ void AShooterNPCSpawner::BeginPlay()
 {
 	Super::BeginPlay();
 	
-	// should we spawn an NPC right away?
-	if (bShouldSpawnNPCsImmediately)
+	// ensure we don't spawn NPCs if our initial spawn count is zero
+	if (SpawnCount > 0)
 	{
 		// schedule the first NPC spawn
 		GetWorld()->GetTimerManager().SetTimer(SpawnTimer, this, &AShooterNPCSpawner::SpawnNPC, InitialSpawnDelay);
diff --git a/Templates/TP_FirstPerson/Source/TP_FirstPerson/Variant_Shooter/AI/ShooterNPCSpawner.h b/Templates/TP_FirstPerson/Source/TP_FirstPerson/Variant_Shooter/AI/ShooterNPCSpawner.h
--- a/Templates/TP_FirstPerson/Source/TP_FirstPerson/Variant_Shooter/AI/ShooterNPCSpawner.h
+++ b/Templates/TP_FirstPerson/Source/TP_FirstPerson/Variant_Shooter/AI/ShooterNPCSpawner.h
@@ -31,10 +31,6 @@ protected:
 	UPROPERTY(EditAnywhere, BlueprintReadOnly, Category="NPC Spawner")
 	TSubclassOf<AShooterNPC> NPCClass;
 
-	/** If true, the first NPC will be spawned as soon as the game starts */
-	UPROPERTY(EditAnywhere, BlueprintReadOnly, Category="NPC Spawner")
-	bool bShouldSpawnNPCsImmediately = true;
-
 	/** Time to wait before spawning the first NPC on game start */
 	UPROPERTY(EditAnywhere, BlueprintReadOnly, Category="NPC Spawner", meta = (ClampMin = 0, ClampMax = 10))
 	float InitialSpawnDelay = 5.0f;
diff --git a/Templates/TP_FirstPerson/Source/TP_FirstPerson/Variant_Shooter/ShooterPlayerController.h b/Templates/TP_FirstPerson/Source/TP_FirstPerson/Variant_Shooter/ShooterPlayerController.h
--- a/Templates/TP_FirstPerson/Source/TP_FirstPerson/Variant_Shooter/ShooterPlayerController.h
+++ b/Templates/TP_FirstPerson/Source/TP_FirstPerson/Variant_Shooter/ShooterPlayerController.h
@@ -35,6 +35,7 @@ protected:
 	TSubclassOf<UUserWidget> MobileControlsWidgetClass;
 
 	/** Pointer to the mobile controls widget */
+	UPROPERTY()
 	TObjectPtr<UUserWidget> MobileControlsWidget;
 
 	/** If true, the player will use UMG touch controls even if not playing on mobile platforms */
diff --git a/Templates/TP_ThirdPerson/Source/TP_ThirdPerson/TP_ThirdPersonPlayerController.h b/Templates/TP_ThirdPerson/Source/TP_ThirdPerson/TP_ThirdPersonPlayerController.h
--- a/Templates/TP_ThirdPerson/Source/TP_ThirdPerson/TP_ThirdPersonPlayerController.h
+++ b/Templates/TP_ThirdPerson/Source/TP_ThirdPerson/TP_ThirdPersonPlayerController.h
@@ -33,6 +33,7 @@ protected:
 	TSubclassOf<UUserWidget> MobileControlsWidgetClass;
 
 	/** Pointer to the mobile controls widget */
+	UPROPERTY()
 	TObjectPtr<UUserWidget> MobileControlsWidget;
 
 	/** If true, the player will use UMG touch controls even if not playing on mobile platforms */
diff --git a/Templates/TP_ThirdPerson/Source/TP_ThirdPerson/Variant_Combat/AI/CombatStateTreeUtility.cpp b/Templates/TP_ThirdPerson/Source/TP_ThirdPerson/Variant_Combat/AI/CombatStateTreeUtility.cpp
--- a/Templates/TP_ThirdPerson/Source/TP_ThirdPerson/Variant_Combat/AI/CombatStateTreeUtility.cpp
+++ b/Templates/TP_ThirdPerson/Source/TP_ThirdPerson/Variant_Combat/AI/CombatStateTreeUtility.cpp
@@ -127,7 +127,7 @@ EStateTreeRunStatus FStateTreeChargedAttackTask::EnterState(FStateTreeExecutionC
 			}
 		);
 
-		// tell the character to do a combo attack
+		// tell the character to do a charged attack
 		InstanceData.Character->DoAIChargedAttack();
 	}
 
@@ -184,7 +184,7 @@ void FStateTreeWaitForLandingTask::ExitState(FStateTreeExecutionContext& Context
 		// get the instance data
 		FInstanceDataType& InstanceData = Context.GetInstanceData(*this);
 
-		// bind the on enemy landed delegate
+		// unbind the on enemy landed delegate
 		InstanceData.Character->OnEnemyLanded.Unbind();
 	}
 }
diff --git a/Templates/TP_ThirdPerson/Source/TP_ThirdPerson/Variant_Combat/AI/CombatStateTreeUtility.h b/Templates/TP_ThirdPerson/Source/TP_ThirdPerson/Variant_Combat/AI/CombatStateTreeUtility.h
--- a/Templates/TP_ThirdPerson/Source/TP_ThirdPerson/Variant_Combat/AI/CombatStateTreeUtility.h
+++ b/Templates/TP_ThirdPerson/Source/TP_ThirdPerson/Variant_Combat/AI/CombatStateTreeUtility.h
@@ -66,7 +66,7 @@ struct FStateTreeIsInDangerConditionInstanceData
 {
 	GENERATED_BODY()
 	
-	/** Character to check grounded status on */
+	/** Character to check danger status on */
 	UPROPERTY(EditAnywhere, Category = "Context")
 	ACombatEnemy* Character;
 
diff --git a/Templates/TP_ThirdPerson/Source/TP_ThirdPerson/Variant_Combat/CombatCharacter.cpp b/Templates/TP_ThirdPerson/Source/TP_ThirdPerson/Variant_Combat/CombatCharacter.cpp
--- a/Templates/TP_ThirdPerson/Source/TP_ThirdPerson/Variant_Combat/CombatCharacter.cpp
+++ b/Templates/TP_ThirdPerson/Source/TP_ThirdPerson/Variant_Combat/CombatCharacter.cpp
@@ -188,7 +188,7 @@ void ACombatCharacter::ComboAttack()
 	ComboCount = 0;
 
 	// notify enemies they are about to be attacked
-	NotifyEnemiesOfAttack();
+	NotifyEnemiesOfIncomingAttack();
 
 	// play the attack montage
 	if (UAnimInstance* AnimInstance = GetMesh()->GetAnimInstance())
@@ -214,7 +214,7 @@ void ACombatCharacter::ChargedAttack()
 	bHasLoopedChargedAttack = false;
 
 	// notify enemies they are about to be attacked
-	NotifyEnemiesOfAttack();
+	NotifyEnemiesOfIncomingAttack();
 
 	// play the charged attack montage
 	if (UAnimInstance* AnimInstance = GetMesh()->GetAnimInstance())
@@ -315,7 +315,7 @@ void ACombatCharacter::CheckCombo()
 			if (ComboCount < ComboSectionNames.Num())
 			{
 				// notify enemies they are about to be attacked
-				NotifyEnemiesOfAttack();
+				NotifyEnemiesOfIncomingAttack();
 
 				// jump to the next combo section
 				if (UAnimInstance* AnimInstance = GetMesh()->GetAnimInstance())
@@ -339,7 +339,7 @@ void ACombatCharacter::CheckChargedAttack()
 	}
 }
 
-void ACombatCharacter::NotifyEnemiesOfAttack()
+void ACombatCharacter::NotifyEnemiesOfIncomingAttack()
 {
 	// sweep for objects in front of the character to be hit by the attack
 	TArray<FHitResult> OutHits;
diff --git a/Templates/TP_ThirdPerson/Source/TP_ThirdPerson/Variant_Combat/CombatCharacter.h b/Templates/TP_ThirdPerson/Source/TP_ThirdPerson/Variant_Combat/CombatCharacter.h
--- a/Templates/TP_ThirdPerson/Source/TP_ThirdPerson/Variant_Combat/CombatCharacter.h
+++ b/Templates/TP_ThirdPerson/Source/TP_ThirdPerson/Variant_Combat/CombatCharacter.h
@@ -273,7 +273,7 @@ public:
 	// ~begin CombatDamageable interface
 
 	/** Notifies nearby enemies that an attack is coming so they can react */
-	void NotifyEnemiesOfAttack();
+	void NotifyEnemiesOfIncomingAttack();
 
 	/** Handles damage and knockback events */
 	virtual void ApplyDamage(float Damage, AActor* DamageCauser, const FVector& DamageLocation, const FVector& DamageImpulse) override;
diff --git a/Templates/TP_ThirdPerson/Source/TP_ThirdPerson/Variant_Combat/CombatPlayerController.h b/Templates/TP_ThirdPerson/Source/TP_ThirdPerson/Variant_Combat/CombatPlayerController.h
--- a/Templates/TP_ThirdPerson/Source/TP_ThirdPerson/Variant_Combat/CombatPlayerController.h
+++ b/Templates/TP_ThirdPerson/Source/TP_ThirdPerson/Variant_Combat/CombatPlayerController.h
@@ -14,7 +14,7 @@ class ACombatCharacter;
  *  Manages input mappings
  *  Respawns the player character at the checkpoint when it's destroyed
  */
-UCLASS(abstract)
+UCLASS(abstract, Config="Game")
 class ACombatPlayerController : public APlayerController
 {
 	GENERATED_BODY()
@@ -34,6 +34,7 @@ protected:
 	TSubclassOf<UUserWidget> MobileControlsWidgetClass;
 
 	/** Pointer to the mobile controls widget */
+	UPROPERTY()
 	TObjectPtr<UUserWidget> MobileControlsWidget;
 
 	/** If true, the player will use UMG touch controls even if not playing on mobile platforms */
diff --git a/Templates/TP_ThirdPerson/Source/TP_ThirdPerson/Variant_Combat/Interfaces/CombatDamageable.h b/Templates/TP_ThirdPerson/Source/TP_ThirdPerson/Variant_Combat/Interfaces/CombatDamageable.h
--- a/Templates/TP_ThirdPerson/Source/TP_ThirdPerson/Variant_Combat/Interfaces/CombatDamageable.h
+++ b/Templates/TP_ThirdPerson/Source/TP_ThirdPerson/Variant_Combat/Interfaces/CombatDamageable.h
@@ -9,6 +9,7 @@
 /**
  *  CombatDamageable interface
  *  Provides functionality to handle damage, healing, knockback and death
+ *  Also provides functionality to warn characters of incoming sources of damage
  */
 UINTERFACE(MinimalAPI, NotBlueprintable)
 class UCombatDamageable : public UInterface
diff --git a/Templates/TP_ThirdPerson/Source/TP_ThirdPerson/Variant_Platforming/PlatformingCharacter.cpp b/Templates/TP_ThirdPerson/Source/TP_ThirdPerson/Variant_Platforming/PlatformingCharacter.cpp
--- a/Templates/TP_ThirdPerson/Source/TP_ThirdPerson/Variant_Platforming/PlatformingCharacter.cpp
+++ b/Templates/TP_ThirdPerson/Source/TP_ThirdPerson/Variant_Platforming/PlatformingCharacter.cpp
@@ -24,7 +24,7 @@ APlatformingCharacter::APlatformingCharacter()
 	bHasDashed = false;
 	bIsDashing = false;
 
-	// bind the attack montage ended delegate
+	// bind the dash montage ended delegate
 	OnDashMontageEnded.BindUObject(this, &APlatformingCharacter::DashMontageEnded);
 
 	// enable press and hold jump
@@ -278,11 +278,8 @@ void APlatformingCharacter::DoJumpEnd()
 
 void APlatformingCharacter::DashMontageEnded(UAnimMontage* Montage, bool bInterrupted)
 {
-	// if the montage was interrupted, end the dash
-	if (bInterrupted)
-	{
-		EndDash();
-	}
+	// end the dash
+	EndDash();
 }
 
 void APlatformingCharacter::EndDash()
diff --git a/Templates/TP_ThirdPerson/Source/TP_ThirdPerson/Variant_Platforming/PlatformingPlayerController.h b/Templates/TP_ThirdPerson/Source/TP_ThirdPerson/Variant_Platforming/PlatformingPlayerController.h
--- a/Templates/TP_ThirdPerson/Source/TP_ThirdPerson/Variant_Platforming/PlatformingPlayerController.h
+++ b/Templates/TP_ThirdPerson/Source/TP_ThirdPerson/Variant_Platforming/PlatformingPlayerController.h
@@ -14,7 +14,7 @@ class APlatformingCharacter;
  *  Manages input mappings
  *  Respawns the player character at the Player Start when it's destroyed
  */
-UCLASS(abstract)
+UCLASS(abstract, Config="Game")
 class APlatformingPlayerController : public APlayerController
 {
 	GENERATED_BODY()
@@ -34,6 +34,7 @@ protected:
 	TSubclassOf<UUserWidget> MobileControlsWidgetClass;
 
 	/** Pointer to the mobile controls widget */
+	UPROPERTY()
 	TObjectPtr<UUserWidget> MobileControlsWidget;
 
 	/** If true, the player will use UMG touch controls even if not playing on mobile platforms */
diff --git a/Templates/TP_ThirdPerson/Source/TP_ThirdPerson/Variant_SideScrolling/SideScrollingPlayerController.h b/Templates/TP_ThirdPerson/Source/TP_ThirdPerson/Variant_SideScrolling/SideScrollingPlayerController.h
--- a/Templates/TP_ThirdPerson/Source/TP_ThirdPerson/Variant_SideScrolling/SideScrollingPlayerController.h
+++ b/Templates/TP_ThirdPerson/Source/TP_ThirdPerson/Variant_SideScrolling/SideScrollingPlayerController.h
@@ -15,7 +15,7 @@ class UInputMappingContext;
  *  Manages input mappings
  *  Respawns the player pawn at the player start if it is destroyed
  */
-UCLASS(abstract)
+UCLASS(abstract, Config="Game")
 class ASideScrollingPlayerController : public APlayerController
 {
 	GENERATED_BODY()
@@ -35,6 +35,7 @@ protected:
 	TSubclassOf<UUserWidget> MobileControlsWidgetClass;
 
 	/** Pointer to the mobile controls widget */
+	UPROPERTY()
 	TObjectPtr<UUserWidget> MobileControlsWidget;
 
 	/** If true, the player will use UMG touch controls even if not playing on mobile platforms */
diff --git a/Templates/TP_TopDown/Source/TP_TopDown/Variant_Strategy/UI/StrategyHUD.h b/Templates/TP_TopDown/Source/TP_TopDown/Variant_Strategy/UI/StrategyHUD.h
--- a/Templates/TP_TopDown/Source/TP_TopDown/Variant_Strategy/UI/StrategyHUD.h
+++ b/Templates/TP_TopDown/Source/TP_TopDown/Variant_Strategy/UI/StrategyHUD.h
@@ -20,6 +20,7 @@ class AStrategyHUD : public AHUD
 protected:
 
 	/** Pointer to the UI user widget */
+	UPROPERTY()
 	TObjectPtr<UStrategyUI> UIWidget;
 
 	/** Type of UI Widget to spawn */
diff --git a/Templates/TP_TopDown/Source/TP_TopDown/Variant_TwinStick/TwinStickPlayerController.h b/Templates/TP_TopDown/Source/TP_TopDown/Variant_TwinStick/TwinStickPlayerController.h
--- a/Templates/TP_TopDown/Source/TP_TopDown/Variant_TwinStick/TwinStickPlayerController.h
+++ b/Templates/TP_TopDown/Source/TP_TopDown/Variant_TwinStick/TwinStickPlayerController.h
@@ -14,7 +14,7 @@ class ATwinStickCharacter;
  *  Manages input mapping contexts
  *  Respawns the pawn if it is destroyed
  */
-UCLASS(abstract)
+UCLASS(abstract, Config="Game")
 class ATwinStickPlayerController : public APlayerController
 {
 	GENERATED_BODY()
@@ -34,6 +34,7 @@ protected:
 	TSubclassOf<UUserWidget> MobileControlsWidgetClass;
 
 	/** Pointer to the mobile controls widget */
+	UPROPERTY()
 	TObjectPtr<UUserWidget> MobileControlsWidget;
 
 	/** If true, the player will use UMG touch controls even if not playing on mobile platforms */
diff --git a/Templates/TP_VehicleAdv/Source/TP_VehicleAdv/TP_VehicleAdvPlayerController.h b/Templates/TP_VehicleAdv/Source/TP_VehicleAdv/TP_VehicleAdvPlayerController.h
--- a/Templates/TP_VehicleAdv/Source/TP_VehicleAdv/TP_VehicleAdvPlayerController.h
+++ b/Templates/TP_VehicleAdv/Source/TP_VehicleAdv/TP_VehicleAdvPlayerController.h
@@ -14,7 +14,7 @@ class UTP_VehicleAdvUI;
  *  Vehicle Player Controller class
  *  Handles input mapping and user interface
  */
-UCLASS(abstract)
+UCLASS(abstract, Config="Game")
 class ATP_VehicleAdvPlayerController : public APlayerController
 {
 	GENERATED_BODY()
@@ -34,6 +34,7 @@ protected:
 	TSubclassOf<UUserWidget> MobileControlsWidgetClass;
 
 	/** Pointer to the mobile controls widget */
+	UPROPERTY()
 	TObjectPtr<UUserWidget> MobileControlsWidget;
 
 	/** If true, the player will use UMG touch controls even if not playing on mobile platforms */
diff --git a/Templates/TP_VehicleAdv/Source/TP_VehicleAdv/Variant_TimeTrial/TimeTrialPlayerController.h b/Templates/TP_VehicleAdv/Source/TP_VehicleAdv/Variant_TimeTrial/TimeTrialPlayerController.h
--- a/Templates/TP_VehicleAdv/Source/TP_VehicleAdv/Variant_TimeTrial/TimeTrialPlayerController.h
+++ b/Templates/TP_VehicleAdv/Source/TP_VehicleAdv/Variant_TimeTrial/TimeTrialPlayerController.h
@@ -15,7 +15,7 @@ class ATP_VehicleAdvPawn;
 /**
  *  A simple PlayerController for a Time Trial racing game
  */
-UCLASS(abstract)
+UCLASS(abstract, Config="Game")
 class ATimeTrialPlayerController : public APlayerController
 {
 	GENERATED_BODY()
@@ -35,6 +35,7 @@ protected:
 	TSubclassOf<UUserWidget> MobileControlsWidgetClass;
 
 	/** Pointer to the mobile controls widget */
+	UPROPERTY()
 	TObjectPtr<UUserWidget> MobileControlsWidget;
 
 	/** If true, the player will use UMG touch controls even if not playing on mobile platforms */
