TP_ThirdPerson
- Fixed input initialization crash on SideScrollingPlayerController
#rnx

diff --git a/Templates/TP_ThirdPerson/Source/TP_ThirdPerson/Variant_SideScrolling/SideScrollingPlayerController.cpp b/Templates/TP_ThirdPerson/Source/TP_ThirdPerson/Variant_SideScrolling/SideScrollingPlayerController.cpp
--- a/Templates/TP_ThirdPerson/Source/TP_ThirdPerson/Variant_SideScrolling/SideScrollingPlayerController.cpp
+++ b/Templates/TP_ThirdPerson/Source/TP_ThirdPerson/Variant_SideScrolling/SideScrollingPlayerController.cpp
@@ -39,6 +39,8 @@ void ASideScrollingPlayerController::BeginPlay()
 
 void ASideScrollingPlayerController::SetupInputComponent()
 {
+	Super::SetupInputComponent();
+
 	// only add IMCs for local player controllers
 	if (IsLocalPlayerController())
 	{
