Fix a deprecation warning.
#rnx

diff --git a/Templates/TP_ThirdPerson/Source/TP_ThirdPerson/Variant_Combat/AI/CombatStateTreeUtility.h b/Templates/TP_ThirdPerson/Source/TP_ThirdPerson/Variant_Combat/AI/CombatStateTreeUtility.h
--- a/Templates/TP_ThirdPerson/Source/TP_ThirdPerson/Variant_Combat/AI/CombatStateTreeUtility.h
+++ b/Templates/TP_ThirdPerson/Source/TP_ThirdPerson/Variant_Combat/AI/CombatStateTreeUtility.h
@@ -67,7 +67,7 @@ struct FStateTreeIsInDangerConditionInstanceData
 	
 	/** Character to check danger status on */
 	UPROPERTY(EditAnywhere, Category = "Context")
-	ACombatEnemy* Character;
+	TObjectPtr<ACombatEnemy> Character;
 
 	/** Minimum time to wait before reacting to the danger event */
 	UPROPERTY(EditAnywhere, Category = "Parameters", meta = (Units = "s"))
@@ -81,7 +81,6 @@ struct FStateTreeIsInDangerConditionInstanceData
 	UPROPERTY(EditAnywhere, Category = "Parameters", meta = (Units = "degrees"))
 	float DangerSightConeAngle = 120.0f;
 };
-STATETREE_POD_INSTANCEDATA(FStateTreeIsInDangerConditionInstanceData);
 
 /**
  *  StateTree condition to check if the character is about to be hit by an attack
