PR #14392: Update templates to check if PlayerController is locally controlled before checking Slate related things

#rnx

#Synced-CL 52000158

diff --git a/Templates/TP_FirstPerson/Source/TP_FirstPerson/TP_FirstPersonPlayerController.cpp b/Templates/TP_FirstPerson/Source/TP_FirstPerson/TP_FirstPersonPlayerController.cpp
--- a/Templates/TP_FirstPerson/Source/TP_FirstPerson/TP_FirstPersonPlayerController.cpp
+++ b/Templates/TP_FirstPerson/Source/TP_FirstPerson/TP_FirstPersonPlayerController.cpp
@@ -22,7 +22,7 @@ void ATP_FirstPersonPlayerController::BeginPlay()
 
 	
 	// only spawn touch controls on local player controllers
-	if (ShouldUseTouchControls() && IsLocalPlayerController())
+	if (IsLocalPlayerController() && ShouldUseTouchControls())
 	{
 		// spawn the mobile controls widget
 		MobileControlsWidget = CreateWidget<UUserWidget>(this, MobileControlsWidgetClass);
diff --git a/Templates/TP_FirstPerson/Source/TP_FirstPerson/Variant_Horror/HorrorPlayerController.cpp b/Templates/TP_FirstPerson/Source/TP_FirstPerson/Variant_Horror/HorrorPlayerController.cpp
--- a/Templates/TP_FirstPerson/Source/TP_FirstPerson/Variant_Horror/HorrorPlayerController.cpp
+++ b/Templates/TP_FirstPerson/Source/TP_FirstPerson/Variant_Horror/HorrorPlayerController.cpp
@@ -22,7 +22,7 @@ void AHorrorPlayerController::BeginPlay()
 	Super::BeginPlay();
 
 	// only spawn touch controls on local player controllers
-	if (ShouldUseTouchControls() && IsLocalPlayerController())
+	if (IsLocalPlayerController() && ShouldUseTouchControls())
 	{
 		// spawn the mobile controls widget
 		MobileControlsWidget = CreateWidget<UUserWidget>(this, MobileControlsWidgetClass);
diff --git a/Templates/TP_ThirdPerson/Source/TP_ThirdPerson/TP_ThirdPersonPlayerController.cpp b/Templates/TP_ThirdPerson/Source/TP_ThirdPerson/TP_ThirdPersonPlayerController.cpp
--- a/Templates/TP_ThirdPerson/Source/TP_ThirdPerson/TP_ThirdPersonPlayerController.cpp
+++ b/Templates/TP_ThirdPerson/Source/TP_ThirdPerson/TP_ThirdPersonPlayerController.cpp
@@ -14,7 +14,7 @@ void ATP_ThirdPersonPlayerController::BeginPlay()
 	Super::BeginPlay();
 
 	// only spawn touch controls on local player controllers
-	if (ShouldUseTouchControls() && IsLocalPlayerController())
+	if (IsLocalPlayerController() && ShouldUseTouchControls())
 	{
 		// spawn the mobile controls widget
 		MobileControlsWidget = CreateWidget<UUserWidget>(this, MobileControlsWidgetClass);
diff --git a/Templates/TP_ThirdPerson/Source/TP_ThirdPerson/Variant_Combat/CombatPlayerController.cpp b/Templates/TP_ThirdPerson/Source/TP_ThirdPerson/Variant_Combat/CombatPlayerController.cpp
--- a/Templates/TP_ThirdPerson/Source/TP_ThirdPerson/Variant_Combat/CombatPlayerController.cpp
+++ b/Templates/TP_ThirdPerson/Source/TP_ThirdPerson/Variant_Combat/CombatPlayerController.cpp
@@ -18,7 +18,7 @@ void ACombatPlayerController::BeginPlay()
 	Super::BeginPlay();
 
 	// only spawn touch controls on local player controllers
-	if (ShouldUseTouchControls() && IsLocalPlayerController())
+	if (IsLocalPlayerController() && ShouldUseTouchControls())
 	{
 		// spawn the mobile controls widget
 		MobileControlsWidget = CreateWidget<UUserWidget>(this, MobileControlsWidgetClass);
diff --git a/Templates/TP_ThirdPerson/Source/TP_ThirdPerson/Variant_Platforming/PlatformingPlayerController.cpp b/Templates/TP_ThirdPerson/Source/TP_ThirdPerson/Variant_Platforming/PlatformingPlayerController.cpp
--- a/Templates/TP_ThirdPerson/Source/TP_ThirdPerson/Variant_Platforming/PlatformingPlayerController.cpp
+++ b/Templates/TP_ThirdPerson/Source/TP_ThirdPerson/Variant_Platforming/PlatformingPlayerController.cpp
@@ -18,7 +18,7 @@ void APlatformingPlayerController::BeginPlay()
 	Super::BeginPlay();
 
 	// only spawn touch controls on local player controllers
-	if (ShouldUseTouchControls() && IsLocalPlayerController())
+	if (IsLocalPlayerController() && ShouldUseTouchControls())
 	{
 		// spawn the mobile controls widget
 		MobileControlsWidget = CreateWidget<UUserWidget>(this, MobileControlsWidgetClass);
diff --git a/Templates/TP_ThirdPerson/Source/TP_ThirdPerson/Variant_SideScrolling/SideScrollingPlayerController.cpp b/Templates/TP_ThirdPerson/Source/TP_ThirdPerson/Variant_SideScrolling/SideScrollingPlayerController.cpp
--- a/Templates/TP_ThirdPerson/Source/TP_ThirdPerson/Variant_SideScrolling/SideScrollingPlayerController.cpp
+++ b/Templates/TP_ThirdPerson/Source/TP_ThirdPerson/Variant_SideScrolling/SideScrollingPlayerController.cpp
@@ -18,7 +18,7 @@ void ASideScrollingPlayerController::BeginPlay()
 	Super::BeginPlay();
 
 	// only spawn touch controls on local player controllers
-	if (ShouldUseTouchControls() && IsLocalPlayerController())
+	if (IsLocalPlayerController() && ShouldUseTouchControls())
 	{
 		// spawn the mobile controls widget
 		MobileControlsWidget = CreateWidget<UUserWidget>(this, MobileControlsWidgetClass);
diff --git a/Templates/TP_TopDown/Source/TP_TopDown/Variant_TwinStick/TwinStickPlayerController.cpp b/Templates/TP_TopDown/Source/TP_TopDown/Variant_TwinStick/TwinStickPlayerController.cpp
--- a/Templates/TP_TopDown/Source/TP_TopDown/Variant_TwinStick/TwinStickPlayerController.cpp
+++ b/Templates/TP_TopDown/Source/TP_TopDown/Variant_TwinStick/TwinStickPlayerController.cpp
@@ -18,7 +18,7 @@ void ATwinStickPlayerController::BeginPlay()
 	Super::BeginPlay();
 
 	// only spawn touch controls on local player controllers
-	if (ShouldUseTouchControls() && IsLocalPlayerController())
+	if (IsLocalPlayerController() && ShouldUseTouchControls())
 	{
 		// spawn the mobile controls widget
 		MobileControlsWidget = CreateWidget<UUserWidget>(this, MobileControlsWidgetClass);
