Templates - TP_ThirdPerson, TP_ThirdPersonBP
- Combat Variant: Adjusted enemy life bar placement.
- Combat Variant: Adjusted and optimized enemy StateTrees
- Side Scroller Variant: Added blocking volumes to limit physics box and NPCs accidentally moving outside of intended areas.
- Side Scroller Variant: Fixed an issue that attempted to apply a physics force to NPCs when colliding with the player while falling or jumping.

#rnx

#virtualized
#Synced-CL 52765683

diff --git a/Templates/TP_ThirdPerson/Source/TP_ThirdPerson/Variant_Combat/AI/CombatStateTreeUtility.cpp b/Templates/TP_ThirdPerson/Source/TP_ThirdPerson/Variant_Combat/AI/CombatStateTreeUtility.cpp
--- a/Templates/TP_ThirdPerson/Source/TP_ThirdPerson/Variant_Combat/AI/CombatStateTreeUtility.cpp
+++ b/Templates/TP_ThirdPerson/Source/TP_ThirdPerson/Variant_Combat/AI/CombatStateTreeUtility.cpp
@@ -67,39 +67,31 @@ FText FStateTreeIsInDangerCondition::GetDescription(const FGuid& ID, FStateTreeD
 
 EStateTreeRunStatus FStateTreeComboAttackTask::EnterState(FStateTreeExecutionContext& Context, const FStateTreeTransitionResult& Transition) const
 {
-	// have we transitioned from another state?
-	if (Transition.ChangeType == EStateTreeStateChangeType::Changed)
-	{
-		// get the instance data
-		FInstanceDataType& InstanceData = Context.GetInstanceData(*this);
+	// get the instance data
+	FInstanceDataType& InstanceData = Context.GetInstanceData(*this);
 
-		// bind to the on attack completed delegate
-		InstanceData.Character->OnAttackCompleted.BindLambda(
-			[WeakContext = Context.MakeWeakExecutionContext()]()
-			{
-				WeakContext.FinishTask(EStateTreeFinishTaskType::Succeeded);
-			}
-		);
+	// bind to the on attack completed delegate
+	InstanceData.Character->OnAttackCompleted.BindLambda(
+		[WeakContext = Context.MakeWeakExecutionContext()]()
+		{
+			WeakContext.FinishTask(EStateTreeFinishTaskType::Succeeded);
+		}
+	);
 
 
-		// tell the character to do a combo attack
-		InstanceData.Character->DoAIComboAttack();
-	}
+	// tell the character to do a combo attack
+	InstanceData.Character->DoAIComboAttack();
 
 	return EStateTreeRunStatus::Running;
 }
 
 void FStateTreeComboAttackTask::ExitState(FStateTreeExecutionContext& Context, const FStateTreeTransitionResult& Transition) const
 {
-	// have we transitioned from another state?
-	if (Transition.ChangeType == EStateTreeStateChangeType::Changed)
-	{
-		// get the instance data
-		FInstanceDataType& InstanceData = Context.GetInstanceData(*this);
+	// get the instance data
+	FInstanceDataType& InstanceData = Context.GetInstanceData(*this);
 
-		// unbind the on attack completed delegate
-		InstanceData.Character->OnAttackCompleted.Unbind();
-	}
+	// unbind the on attack completed delegate
+	InstanceData.Character->OnAttackCompleted.Unbind();
 }
 
 #if WITH_EDITOR
@@ -113,38 +105,30 @@ FText FStateTreeComboAttackTask::GetDescription(const FGuid& ID, FStateTreeDataV
 
 EStateTreeRunStatus FStateTreeChargedAttackTask::EnterState(FStateTreeExecutionContext& Context, const FStateTreeTransitionResult& Transition) const
 {
-	// have we transitioned from another state?
-	if (Transition.ChangeType == EStateTreeStateChangeType::Changed)
-	{
-		// get the instance data
-		FInstanceDataType& InstanceData = Context.GetInstanceData(*this);
-
-		// bind to the on attack completed delegate
-		InstanceData.Character->OnAttackCompleted.BindLambda(
-			[WeakContext = Context.MakeWeakExecutionContext()]()
-			{
-				WeakContext.FinishTask(EStateTreeFinishTaskType::Succeeded);
-			}
-		);
-
-		// tell the character to do a charged attack
-		InstanceData.Character->DoAIChargedAttack();
-	}
+	// get the instance data
+	FInstanceDataType& InstanceData = Context.GetInstanceData(*this);
+
+	// bind to the on attack completed delegate
+	InstanceData.Character->OnAttackCompleted.BindLambda(
+		[WeakContext = Context.MakeWeakExecutionContext()]()
+		{
+			WeakContext.FinishTask(EStateTreeFinishTaskType::Succeeded);
+		}
+	);
+
+	// tell the character to do a charged attack
+	InstanceData.Character->DoAIChargedAttack();
 
 	return EStateTreeRunStatus::Running;
 }
 
 void FStateTreeChargedAttackTask::ExitState(FStateTreeExecutionContext& Context, const FStateTreeTransitionResult& Transition) const
 {
-	// have we transitioned from another state?
-	if (Transition.ChangeType == EStateTreeStateChangeType::Changed)
-	{
-		// get the instance data
-		FInstanceDataType& InstanceData = Context.GetInstanceData(*this);
+	// get the instance data
+	FInstanceDataType& InstanceData = Context.GetInstanceData(*this);
 
-		// unbind the on attack completed delegate
-		InstanceData.Character->OnAttackCompleted.Unbind();
-	}
+	// unbind the on attack completed delegate
+	InstanceData.Character->OnAttackCompleted.Unbind();
 }
 
 #if WITH_EDITOR
@@ -158,35 +142,27 @@ FText FStateTreeChargedAttackTask::GetDescription(const FGuid& ID, FStateTreeDat
 
 EStateTreeRunStatus FStateTreeWaitForLandingTask::EnterState(FStateTreeExecutionContext& Context, const FStateTreeTransitionResult& Transition) const
 {
-	// have we transitioned from another state?
-	if (Transition.ChangeType == EStateTreeStateChangeType::Changed)
-	{
-		// get the instance data
-		FInstanceDataType& InstanceData = Context.GetInstanceData(*this);
-
-		// bind to the on enemy landed delegate
-		InstanceData.Character->OnEnemyLanded.BindLambda(
-			[WeakContext = Context.MakeWeakExecutionContext()]()
-			{
-				WeakContext.FinishTask(EStateTreeFinishTaskType::Succeeded);
-			}
-		);
-	}
+	// get the instance data
+	FInstanceDataType& InstanceData = Context.GetInstanceData(*this);
+
+	// bind to the on enemy landed delegate
+	InstanceData.Character->OnEnemyLanded.BindLambda(
+		[WeakContext = Context.MakeWeakExecutionContext()]()
+		{
+			WeakContext.FinishTask(EStateTreeFinishTaskType::Succeeded);
+		}
+	);
 
 	return EStateTreeRunStatus::Running;
 }
 
 void FStateTreeWaitForLandingTask::ExitState(FStateTreeExecutionContext& Context, const FStateTreeTransitionResult& Transition) const
 {
-	// have we transitioned from another state?
-	if (Transition.ChangeType == EStateTreeStateChangeType::Changed)
-	{
-		// get the instance data
-		FInstanceDataType& InstanceData = Context.GetInstanceData(*this);
+	// get the instance data
+	FInstanceDataType& InstanceData = Context.GetInstanceData(*this);
 
-		// unbind the on enemy landed delegate
-		InstanceData.Character->OnEnemyLanded.Unbind();
-	}
+	// unbind the on enemy landed delegate
+	InstanceData.Character->OnEnemyLanded.Unbind();
 }
 
 #if WITH_EDITOR
@@ -200,30 +176,22 @@ FText FStateTreeWaitForLandingTask::GetDescription(const FGuid& ID, FStateTreeDa
 
 EStateTreeRunStatus FStateTreeFaceActorTask::EnterState(FStateTreeExecutionContext& Context, const FStateTreeTransitionResult& Transition) const
 {
-	// have we transitioned from another state?
-	if (Transition.ChangeType == EStateTreeStateChangeType::Changed)
-	{
-		// get the instance data
-		FInstanceDataType& InstanceData = Context.GetInstanceData(*this);
+	// get the instance data
+	FInstanceDataType& InstanceData = Context.GetInstanceData(*this);
 
-		// set the AI Controller's focus
-		InstanceData.Controller->SetFocus(InstanceData.ActorToFaceTowards);
-	}
+	// set the AI Controller's focus
+	InstanceData.Controller->SetFocus(InstanceData.ActorToFaceTowards);
 
 	return EStateTreeRunStatus::Running;
 }
 
 void FStateTreeFaceActorTask::ExitState(FStateTreeExecutionContext& Context, const FStateTreeTransitionResult& Transition) const
 {
-	// have we transitioned to another state?
-	if (Transition.ChangeType == EStateTreeStateChangeType::Changed)
-	{
-		// get the instance data
-		FInstanceDataType& InstanceData = Context.GetInstanceData(*this);
+	// get the instance data
+	FInstanceDataType& InstanceData = Context.GetInstanceData(*this);
 
-		// clear the AI Controller's focus
-		InstanceData.Controller->ClearFocus(EAIFocusPriority::Gameplay);
-	}
+	// clear the AI Controller's focus
+	InstanceData.Controller->ClearFocus(EAIFocusPriority::Gameplay);
 }
 
 #if WITH_EDITOR
@@ -237,30 +205,22 @@ FText FStateTreeFaceActorTask::GetDescription(const FGuid& ID, FStateTreeDataVie
 
 EStateTreeRunStatus FStateTreeFaceLocationTask::EnterState(FStateTreeExecutionContext& Context, const FStateTreeTransitionResult& Transition) const
 {
-	// have we transitioned from another state?
-	if (Transition.ChangeType == EStateTreeStateChangeType::Changed)
-	{
-		// get the instance data
-		FInstanceDataType& InstanceData = Context.GetInstanceData(*this);
+	// get the instance data
+	FInstanceDataType& InstanceData = Context.GetInstanceData(*this);
 
-		// set the AI Controller's focus
-		InstanceData.Controller->SetFocalPoint(InstanceData.FaceLocation);
-	}
+	// set the AI Controller's focus
+	InstanceData.Controller->SetFocalPoint(InstanceData.FaceLocation);
 
 	return EStateTreeRunStatus::Running;
 }
 
 void FStateTreeFaceLocationTask::ExitState(FStateTreeExecutionContext& Context, const FStateTreeTransitionResult& Transition) const
 {
-	// have we transitioned to another state?
-	if (Transition.ChangeType == EStateTreeStateChangeType::Changed)
-	{
-		// get the instance data
-		FInstanceDataType& InstanceData = Context.GetInstanceData(*this);
+	// get the instance data
+	FInstanceDataType& InstanceData = Context.GetInstanceData(*this);
 
-		// clear the AI Controller's focus
-		InstanceData.Controller->ClearFocus(EAIFocusPriority::Gameplay);
-	}
+	// clear the AI Controller's focus
+	InstanceData.Controller->ClearFocus(EAIFocusPriority::Gameplay);
 }
 
 #if WITH_EDITOR
@@ -274,15 +234,11 @@ FText FStateTreeFaceLocationTask::GetDescription(const FGuid& ID, FStateTreeData
 
 EStateTreeRunStatus FStateTreeSetCharacterSpeedTask::EnterState(FStateTreeExecutionContext& Context, const FStateTreeTransitionResult& Transition) const
 {
-	// have we transitioned from another state?
-	if (Transition.ChangeType == EStateTreeStateChangeType::Changed)
-	{
-		// get the instance data
-		FInstanceDataType& InstanceData = Context.GetInstanceData(*this);
+	// get the instance data
+	FInstanceDataType& InstanceData = Context.GetInstanceData(*this);
 
-		// set the character's max ground speed
-		InstanceData.Character->GetCharacterMovement()->MaxWalkSpeed = InstanceData.Speed;
-	}
+	// set the character's max ground speed
+	InstanceData.Character->GetCharacterMovement()->MaxWalkSpeed = InstanceData.Speed;
 
 	return EStateTreeRunStatus::Running;
 }
diff --git a/Templates/TP_ThirdPerson/Source/TP_ThirdPerson/Variant_Combat/AI/CombatStateTreeUtility.h b/Templates/TP_ThirdPerson/Source/TP_ThirdPerson/Variant_Combat/AI/CombatStateTreeUtility.h
--- a/Templates/TP_ThirdPerson/Source/TP_ThirdPerson/Variant_Combat/AI/CombatStateTreeUtility.h
+++ b/Templates/TP_ThirdPerson/Source/TP_ThirdPerson/Variant_Combat/AI/CombatStateTreeUtility.h
@@ -131,6 +131,16 @@ struct FStateTreeComboAttackTask : public FStateTreeTaskCommonBase
 {
 	GENERATED_BODY()
 
+	/** Constructor */
+	FStateTreeComboAttackTask()
+	{
+		// disable tick
+		bShouldCallTick = false;
+
+		// skip state change events if this is sustained
+		bShouldStateChangeOnReselect = false;
+	}
+
 	/* Ensure we're using the correct instance data struct */
 	using FInstanceDataType = FStateTreeAttackInstanceData;
 	virtual const UStruct* GetInstanceDataType() const override { return FInstanceDataType::StaticStruct(); }
@@ -154,6 +164,16 @@ struct FStateTreeChargedAttackTask : public FStateTreeTaskCommonBase
 {
 	GENERATED_BODY()
 
+	/** Constructor */
+	FStateTreeChargedAttackTask()
+	{
+		// disable tick
+		bShouldCallTick = false;
+
+		// skip state change events if this is sustained
+		bShouldStateChangeOnReselect = false;
+	}
+
 	/* Ensure we're using the correct instance data struct */
 	using FInstanceDataType = FStateTreeAttackInstanceData;
 	virtual const UStruct* GetInstanceDataType() const override { return FInstanceDataType::StaticStruct(); }
@@ -177,6 +197,16 @@ struct FStateTreeWaitForLandingTask : public FStateTreeTaskCommonBase
 {
 	GENERATED_BODY()
 
+	/** Constructor */
+	FStateTreeWaitForLandingTask()
+	{
+		// disable tick
+		bShouldCallTick = false;
+
+		// skip state change events if this is sustained
+		bShouldStateChangeOnReselect = false;
+	}
+
 	/* Ensure we're using the correct instance data struct */
 	using FInstanceDataType = FStateTreeAttackInstanceData;
 	virtual const UStruct* GetInstanceDataType() const override { return FInstanceDataType::StaticStruct(); }
@@ -219,6 +249,16 @@ struct FStateTreeFaceActorTask : public FStateTreeTaskCommonBase
 {
 	GENERATED_BODY()
 
+	/** Constructor */
+	FStateTreeFaceActorTask()
+	{
+		// disable tick
+		bShouldCallTick = false;
+
+		// skip state change events if this is sustained
+		bShouldStateChangeOnReselect = false;
+	}
+
 	/* Ensure we're using the correct instance data struct */
 	using FInstanceDataType = FStateTreeFaceActorInstanceData;
 	virtual const UStruct* GetInstanceDataType() const override { return FInstanceDataType::StaticStruct(); }
@@ -261,6 +301,16 @@ struct FStateTreeFaceLocationTask : public FStateTreeTaskCommonBase
 {
 	GENERATED_BODY()
 
+	/** Constructor */
+	FStateTreeFaceLocationTask()
+	{
+		// disable tick
+		bShouldCallTick = false;
+
+		// skip state change events if this is sustained
+		bShouldStateChangeOnReselect = false;
+	}
+
 	/* Ensure we're using the correct instance data struct */
 	using FInstanceDataType = FStateTreeFaceLocationInstanceData;
 	virtual const UStruct* GetInstanceDataType() const override { return FInstanceDataType::StaticStruct(); }
@@ -303,6 +353,16 @@ struct FStateTreeSetCharacterSpeedTask : public FStateTreeTaskCommonBase
 {
 	GENERATED_BODY()
 
+	/** Constructor */
+	FStateTreeSetCharacterSpeedTask()
+	{
+		// disable tick
+		bShouldCallTick = false;
+
+		// skip state change events if this is sustained
+		bShouldStateChangeOnReselect = false;
+	}
+
 	/* Ensure we're using the correct instance data struct */
 	using FInstanceDataType = FStateTreeSetCharacterSpeedInstanceData;
 	virtual const UStruct* GetInstanceDataType() const override { return FInstanceDataType::StaticStruct(); }
@@ -350,6 +410,16 @@ struct FStateTreeGetPlayerInfoTask : public FStateTreeTaskCommonBase
 {
 	GENERATED_BODY()
 
+	/** Constructor */
+	FStateTreeGetPlayerInfoTask()
+	{
+		// enable tick
+		bShouldCallTick = true;
+
+		// skip state change events if this is sustained
+		bShouldStateChangeOnReselect = false;
+	}
+
 	/* Ensure we're using the correct instance data struct */
 	using FInstanceDataType = FStateTreeGetPlayerInfoInstanceData;
 	virtual const UStruct* GetInstanceDataType() const override { return FInstanceDataType::StaticStruct(); }
